Snapshot for changelist 2348267.

[CL 2348267 in Main branch]

diff --git a/Engine/Extras/VisualStudioSnippets/README.md b/Engine/Extras/VisualStudioSnippets/README.md
new file mode 100644
--- /dev/null
+++ b/Engine/Extras/VisualStudioSnippets/README.md
@@ -0,0 +1,67 @@
+Visual Studio Snippets for Unreal Engine C++ Projects
+===========
+
+How to install snippets?
+===========
+
+**Method One**
+Paste .snippet files into: [REDACTED_USER_HOME]\Documents\Visual Studio 2013\Code Snippets\Visual C++\My Code Snippets. Then restart VS.
+
+**Method Two**
+Open Visual Studio, navigate to TOOLS -> Code Snippets Manager… -> Import…
+
+How to use snippets?
+===========
+
+**Method One**
+Just start typing ue4... snippet list should be loaded in a form of combo box. Then use arrows to select snippet. Hit ENTER or TAB to insert snippet.
+
+**Method Two**
+Type all snippet name and hit TAB. You don't have to wait for VS to show snippet list.
+
+To navigate between highlighted fields you can use TAB and SHIFT + TAB. After you enter all names, hit ENTER.
+
+Snippets
+===========
+
+*ue4classa* – Blueprintable class that derives from an AActor. Parameters are: comment, class name and base class name.
+
+*ue4classu* – Blueprintable class that derives from an UObject. Parameters are: comment, class name, base class name.
+
+*ue4struct* – Simple structure. Parameters are: comment and name.
+
+*ue4interface* – Simple ue4 interface. Parameters are: comment and name.
+
+*ue4bpevent* – This function can be used as an event in blueprint. Parameters are: comment, UI category, virtual and const modifiers, function name and arguments.
+
+*ue4bpfunc* – This function is available for blueprint logic. Parameters are: comment (parameters and return value), UI category, virtual and const modifiers, function name and arguments.
+
+*ue4prop* – This read/write property is available everywhere (blueprint, instance and archetype details). Parameters are: comment, category, type and name.
+
+*ue4enum* – Simple enum. Parameters are: comment, enum name, first member name and it’s comment.
+
+*ue4enumdisplay* – Enum that can be used with blueprints. Parameters are: comment, enum name, first member name, it’s display name and comment.
+
+*ue4log* – Simplest log line. Parameters are category, verbosity and message.
+
+*ue4logdeclare* – Declaration of log category. Place this in main header of your project to allow logging. Parameters are: category, default verbosity and compile time verbosity.
+
+*ue4logdefine* – Definition of log category. Place this in main code file. Parameter is category name.
+
+*ue4logfloat* – Log line that can be used to print float value. Parameters are: category, verbosity and variable name.
+
+*ue4logint* – This log line can be used to log an integer value. Parameters are: category, verbosity and variable name.
+
+*ue4loguobj* – This log line is designed to log from inside of the objects. By default, square brackets contains a name of an object that writes the log. Parameters are: category, verbosity, message and name of a pointer to the object.
+
+*ue4mark* – Can be used to mark changes in engine classes. Parameters are: Company symbol, task/ticket number, name and surname of a developer and short description of modification.
+
+*ue4eve* - 9 snippets for each params combination. Can be used to create event. Parameters are: owning type and event type name.
+
+*ue4del* - 9 snippets for each params combination. Can be used to create delegate. Parameters are: delegate type name and param type names.
+
+*ue4delmul* - 9 snippets for each params combination. Can be used to create multicast delegate. Parameters are: delegate type name and param type names.
+
+*ue4deldyn* - 9 snippets for each params combination. Can be used to create dynamic delegate. Parameters are: delegate type name, param type names and display values.
+
+*ue4deldynmul* - 9 snippets for each params combination. Can be used to create dynamic multicast delegate. Parameters are: delegate type name, param type names and display values`.
\ No newline at end of file
diff --git a/Engine/Extras/VisualStudioSnippets/UE4.SingleLineTrace.snippet b/Engine/Extras/VisualStudioSnippets/UE4.SingleLineTrace.snippet
new file mode 100644
--- /dev/null
+++ b/Engine/Extras/VisualStudioSnippets/UE4.SingleLineTrace.snippet
@@ -0,0 +1,26 @@
+<?xml version="1.0" encoding="utf-8" ?>
+<CodeSnippets xmlns="http://schemas.microsoft.com/VisualStudio/2005/CodeSnippet">
+   <CodeSnippet Format="1.0.0">
+      <Header>
+         <Title>UE4-SingleLineTrace</Title>
+         <Shortcut>SingleLineTrace</Shortcut>
+         <Description>UE4: A Single Line Trace (ignoring self) against World Dynamic</Description>
+         <Author>Epic Games</Author>
+         <SnippetTypes>
+            <SnippetType>Expansion</SnippetType>
+         </SnippetTypes>
+      </Header>
+      <Snippet>
+         <Code Language="cpp"><![CDATA[FCollisionQueryParams TraceParams(NAME_None, /*TraceComplex*/ true, /*IgnoreActor*/ this);
+		FCollisionObjectQueryParams CollisionParams(ECC_WorldDynamic);
+
+		FHitResult OutHit;
+		const bool bHit = World->LineTraceSingle(OutHit, StartPos, EndPos, TraceParams, CollisionParams);
+		if (bHit)
+		{
+			// OutHit Contains Hit Data...
+		}]]>
+         </Code>
+      </Snippet>
+   </CodeSnippet>
+</CodeSnippets>
\ No newline at end of file
diff --git a/Engine/Extras/VisualStudioSnippets/UE4.UFunction.snippet b/Engine/Extras/VisualStudioSnippets/UE4.UFunction.snippet
new file mode 100644
--- /dev/null
+++ b/Engine/Extras/VisualStudioSnippets/UE4.UFunction.snippet
@@ -0,0 +1,21 @@
+<?xml version="1.0" encoding="utf-8" ?>
+<CodeSnippets xmlns="http://schemas.microsoft.com/VisualStudio/2005/CodeSnippet">
+   <CodeSnippet Format="1.0.0">
+      <Header>
+         <Title>UE4-UFunction</Title>
+         <Shortcut>ufunction</Shortcut>
+         <Description>UE4: A Template UFunction</Description>
+         <Author>Epic Games</Author>
+         <SnippetTypes>
+            <SnippetType>Expansion</SnippetType>
+         </SnippetTypes>
+      </Header>
+      <Snippet>
+         <Code Language="cpp"><![CDATA[/** Does the thing. */
+		UFUNCTION(BlueprintCallable, Category="Game Play")
+		void DoThing();
+		]]>
+         </Code>
+      </Snippet>
+   </CodeSnippet>
+</CodeSnippets>
\ No newline at end of file
diff --git a/Engine/Extras/VisualStudioSnippets/UE4.UProperty.snippet b/Engine/Extras/VisualStudioSnippets/UE4.UProperty.snippet
new file mode 100644
--- /dev/null
+++ b/Engine/Extras/VisualStudioSnippets/UE4.UProperty.snippet
@@ -0,0 +1,21 @@
+<?xml version="1.0" encoding="utf-8" ?>
+<CodeSnippets xmlns="http://schemas.microsoft.com/VisualStudio/2005/CodeSnippet">
+   <CodeSnippet Format="1.0.0">
+      <Header>
+         <Title>UE4-UProperty</Title>
+         <Shortcut>uproperty</Shortcut>
+         <Description>UE4: A Template UProperty</Description>
+         <Author>Epic Games</Author>
+         <SnippetTypes>
+            <SnippetType>Expansion</SnippetType>
+         </SnippetTypes>
+      </Header>
+      <Snippet>
+         <Code Language="cpp"><![CDATA[/** The MyProperty */
+		UPROPERTY(EditAnywhere, BlueprintReadWrite, Category="Game")
+		float MyProperty;
+		]]>
+         </Code>
+      </Snippet>
+   </CodeSnippet>
+</CodeSnippets>
\ No newline at end of file
diff --git a/Engine/Extras/VisualStudioSnippets/ue4bpevent.snippet b/Engine/Extras/VisualStudioSnippets/ue4bpevent.snippet
new file mode 100644
--- /dev/null
+++ b/Engine/Extras/VisualStudioSnippets/ue4bpevent.snippet
@@ -0,0 +1,76 @@
+﻿<?xml version="1.0" encoding="utf-8"?>
+<CodeSnippets xmlns="http://schemas.microsoft.com/VisualStudio/2005/CodeSnippet">
+  <CodeSnippet Format="1.0.0">
+    <Header>
+      <SnippetTypes>
+        <SnippetType>Expansion</SnippetType>
+      </SnippetTypes>
+      <Title>ue4bpevent</Title>
+      <Author>Epic Games</Author>
+      <Description>Unreal Engine blueprint implementable event snippet</Description>
+      <HelpUrl>
+      </HelpUrl>
+      <Shortcut>ue4bpevent</Shortcut>
+    </Header>
+    <Snippet>
+      <Declarations>
+		<Literal Editable="true">
+          <ID>comment</ID>
+          <ToolTip>Function comment</ToolTip>
+          <Default>Comment</Default>
+          <Function>
+          </Function>
+        </Literal>
+		<Literal Editable="true">
+          <ID>commentParam</ID>
+          <ToolTip>Function comment</ToolTip>
+          <Default>Comment</Default>
+          <Function>
+          </Function>
+        </Literal>
+        <Literal Editable="true">
+          <ID>category</ID>
+          <ToolTip>UI category name</ToolTip>
+          <Default>MyProject</Default>
+          <Function>
+          </Function>
+        </Literal>
+		<Literal Editable="true">
+          <ID>virtual</ID>
+          <ToolTip>Virtual modifier</ToolTip>
+          <Default>virtual</Default>
+          <Function>
+          </Function>
+        </Literal>
+		<Literal Editable="true">
+          <ID>name</ID>
+          <ToolTip>Function name</ToolTip>
+          <Default>FunctionName</Default>
+          <Function>
+          </Function>
+        </Literal>
+		<Literal Editable="true">
+          <ID>args</ID>
+          <ToolTip>Arguments</ToolTip>
+          <Default>args</Default>
+          <Function>
+          </Function>
+        </Literal>
+		<Literal Editable="true">
+          <ID>const</ID>
+          <ToolTip>Const function modifier</ToolTip>
+          <Default>const</Default>
+          <Function>
+          </Function>
+        </Literal>
+      </Declarations>
+      <Code Language="cpp"><![CDATA[/**
+ * $comment$ 
+ * @param $commentParam$
+ */
+UFUNCTION(BlueprintImplementableEvent, Category="$category$")
+$virtual$ void On$name$($args$) $const$;
+$selected$ $end$]]></Code>
+    </Snippet>
+  </CodeSnippet>
+</CodeSnippets>
\ No newline at end of file
diff --git a/Engine/Extras/VisualStudioSnippets/ue4bpfunc.snippet b/Engine/Extras/VisualStudioSnippets/ue4bpfunc.snippet
new file mode 100644
--- /dev/null
+++ b/Engine/Extras/VisualStudioSnippets/ue4bpfunc.snippet
@@ -0,0 +1,91 @@
+﻿<?xml version="1.0" encoding="utf-8"?>
+<CodeSnippets xmlns="http://schemas.microsoft.com/VisualStudio/2005/CodeSnippet">
+  <CodeSnippet Format="1.0.0">
+    <Header>
+      <SnippetTypes>
+        <SnippetType>Expansion</SnippetType>
+      </SnippetTypes>
+      <Title>ue4bpfunc</Title>
+      <Author>Epic Games</Author>
+      <Description>Unreal Engine blueprint callable function snippet</Description>
+      <HelpUrl>
+      </HelpUrl>
+      <Shortcut>ue4bpfunc</Shortcut>
+    </Header>
+    <Snippet>
+      <Declarations>
+		<Literal Editable="true">
+          <ID>comment</ID>
+          <ToolTip>Function comment</ToolTip>
+          <Default>Comment</Default>
+          <Function>
+          </Function>
+        </Literal>
+		<Literal Editable="true">
+          <ID>commentParam</ID>
+          <ToolTip>Function comment</ToolTip>
+          <Default>Comment</Default>
+          <Function>
+          </Function>
+        </Literal>
+		<Literal Editable="true">
+          <ID>commentReturn</ID>
+          <ToolTip>Function comment</ToolTip>
+          <Default>Comment</Default>
+          <Function>
+          </Function>
+        </Literal>
+        <Literal Editable="true">
+          <ID>category</ID>
+          <ToolTip>UI category name</ToolTip>
+          <Default>MyProject</Default>
+          <Function>
+          </Function>
+        </Literal>
+		<Literal Editable="true">
+          <ID>virtual</ID>
+          <ToolTip>Virtual modifier</ToolTip>
+          <Default>virtual</Default>
+          <Function>
+          </Function>
+        </Literal>
+		<Literal Editable="true">
+          <ID>returnType</ID>
+          <ToolTip>Function return type</ToolTip>
+          <Default>void</Default>
+          <Function>
+          </Function>
+        </Literal>
+		<Literal Editable="true">
+          <ID>name</ID>
+          <ToolTip>Function name</ToolTip>
+          <Default>FunctionName</Default>
+          <Function>
+          </Function>
+        </Literal>
+		<Literal Editable="true">
+          <ID>args</ID>
+          <ToolTip>Arguments</ToolTip>
+          <Default>args</Default>
+          <Function>
+          </Function>
+        </Literal>
+		<Literal Editable="true">
+          <ID>const</ID>
+          <ToolTip>Const function modifier</ToolTip>
+          <Default>const</Default>
+          <Function>
+          </Function>
+        </Literal>
+      </Declarations>
+      <Code Language="cpp"><![CDATA[/**
+ * $comment$ 
+ * @param $commentParam$
+ * @return $commentReturn$
+ */
+UFUNCTION(BlueprintCallable, Category="$category$")
+$virtual$ $returnType$ $name$($args$) $const$;
+$selected$ $end$]]></Code>
+    </Snippet>
+  </CodeSnippet>
+</CodeSnippets>
\ No newline at end of file
diff --git a/Engine/Extras/VisualStudioSnippets/ue4bpprop.snippet b/Engine/Extras/VisualStudioSnippets/ue4bpprop.snippet
new file mode 100644
--- /dev/null
+++ b/Engine/Extras/VisualStudioSnippets/ue4bpprop.snippet
@@ -0,0 +1,68 @@
+﻿<?xml version="1.0" encoding="utf-8"?>
+<CodeSnippets xmlns="http://schemas.microsoft.com/VisualStudio/2005/CodeSnippet">
+  <CodeSnippet Format="1.0.0">
+    <Header>
+      <SnippetTypes>
+        <SnippetType>Expansion</SnippetType>
+      </SnippetTypes>
+      <Title>ue4bpprop</Title>
+      <Author>Epic Games</Author>
+      <Description>Unreal Engine blueprint property snippet</Description>
+      <HelpUrl>
+      </HelpUrl>
+      <Shortcut>ue4bpprop</Shortcut>
+    </Header>
+    <Snippet>
+      <Declarations>
+		<Literal Editable="true">
+          <ID>comment</ID>
+          <ToolTip>Property comment</ToolTip>
+          <Default>Comment</Default>
+          <Function>
+          </Function>
+        </Literal>
+        <Literal Editable="true">
+          <ID>category</ID>
+          <ToolTip>UI category name</ToolTip>
+          <Default>MyProject</Default>
+          <Function>
+          </Function>
+        </Literal>
+		<Literal Editable="true">
+          <ID>defaultsAttribute</ID>
+          <ToolTip>Defaults attribute</ToolTip>
+          <Default>EditAnywhere</Default>
+          <Function>
+          </Function>
+        </Literal>
+		<Literal Editable="true">
+          <ID>instanceAttribute</ID>
+          <ToolTip>Instance attribute</ToolTip>
+          <Default>BlueprintReadWrite</Default>
+          <Function>
+          </Function>
+        </Literal>
+		<Literal Editable="true">
+          <ID>type</ID>
+          <ToolTip>Property type</ToolTip>
+          <Default>Type</Default>
+          <Function>
+          </Function>
+        </Literal>
+		<Literal Editable="true">
+          <ID>name</ID>
+          <ToolTip>Property name</ToolTip>
+          <Default>Name</Default>
+          <Function>
+          </Function>
+        </Literal>
+      </Declarations>
+      <Code Language="cpp"><![CDATA[/** 
+ * $comment$ 
+ */
+UPROPERTY($defaultsAttribute$, $instanceAttribute$, Category="$category$")
+$type$ $name$;
+$selected$ $end$]]></Code>
+    </Snippet>
+  </CodeSnippet>
+</CodeSnippets>
\ No newline at end of file
diff --git a/Engine/Extras/VisualStudioSnippets/ue4classa.snippet b/Engine/Extras/VisualStudioSnippets/ue4classa.snippet
new file mode 100644
--- /dev/null
+++ b/Engine/Extras/VisualStudioSnippets/ue4classa.snippet
@@ -0,0 +1,58 @@
+﻿<?xml version="1.0" encoding="utf-8"?>
+<CodeSnippets xmlns="http://schemas.microsoft.com/VisualStudio/2005/CodeSnippet">
+  <CodeSnippet Format="1.0.0">
+    <Header>
+      <SnippetTypes>
+        <SnippetType>Expansion</SnippetType>
+      </SnippetTypes>
+      <Title>ue4classa</Title>
+      <Author>Epic Games</Author>
+      <Description>Unreal Engine class snippet</Description>
+      <HelpUrl>
+      </HelpUrl>
+      <Shortcut>ue4classa</Shortcut>
+    </Header>
+    <Snippet>
+      <Declarations>
+        <Literal Editable="true">
+          <ID>comment</ID>
+          <ToolTip>Class comment</ToolTip>
+          <Default>Comment</Default>
+          <Function>
+          </Function>
+        </Literal>
+		<Literal Editable="true">
+          <ID>className</ID>
+          <ToolTip>Class name</ToolTip>
+          <Default>Foo</Default>
+          <Function>
+          </Function>
+        </Literal>
+		<Literal Editable="true">
+          <ID>baseClassName</ID>
+          <ToolTip>Base class name</ToolTip>
+          <Default>FooBase</Default>
+          <Function>
+          </Function>
+        </Literal>
+      </Declarations>
+      <Code Language="cpp"><![CDATA[/**
+ * $comment$
+ */
+UCLASS(BlueprintType, Blueprintable)
+class A$className$ : public A$baseClassName$
+{
+	GENERATED_UCLASS_BODY()
+	
+public:
+	$selected$ $end$
+	
+protected:
+
+private:
+
+};
+]]></Code>
+    </Snippet>
+  </CodeSnippet>
+</CodeSnippets>
\ No newline at end of file
diff --git a/Engine/Extras/VisualStudioSnippets/ue4classu.snippet b/Engine/Extras/VisualStudioSnippets/ue4classu.snippet
new file mode 100644
--- /dev/null
+++ b/Engine/Extras/VisualStudioSnippets/ue4classu.snippet
@@ -0,0 +1,58 @@
+﻿<?xml version="1.0" encoding="utf-8"?>
+<CodeSnippets xmlns="http://schemas.microsoft.com/VisualStudio/2005/CodeSnippet">
+  <CodeSnippet Format="1.0.0">
+    <Header>
+      <SnippetTypes>
+        <SnippetType>Expansion</SnippetType>
+      </SnippetTypes>
+      <Title>ue4classu</Title>
+      <Author>Epic Games</Author>
+      <Description>Unreal Engine class snippet</Description>
+      <HelpUrl>
+      </HelpUrl>
+      <Shortcut>ue4classu</Shortcut>
+    </Header>
+    <Snippet>
+      <Declarations>
+        <Literal Editable="true">
+          <ID>comment</ID>
+          <ToolTip>Class comment</ToolTip>
+          <Default>Comment</Default>
+          <Function>
+          </Function>
+        </Literal>
+		<Literal Editable="true">
+          <ID>className</ID>
+          <ToolTip>Class name</ToolTip>
+          <Default>Foo</Default>
+          <Function>
+          </Function>
+        </Literal>
+		<Literal Editable="true">
+          <ID>baseClassName</ID>
+          <ToolTip>Base class name</ToolTip>
+          <Default>FooBase</Default>
+          <Function>
+          </Function>
+        </Literal>
+      </Declarations>
+      <Code Language="cpp"><![CDATA[/**
+ *$comment$
+ */
+UCLASS(BlueprintType, Blueprintable)
+class U$className$ : public U$baseClassName$
+{
+	GENERATED_UCLASS_BODY()
+
+public:
+	$selected$ $end$
+	
+protected:
+
+private:
+
+};
+]]></Code>
+    </Snippet>
+  </CodeSnippet>
+</CodeSnippets>
\ No newline at end of file
diff --git a/Engine/Extras/VisualStudioSnippets/ue4del.snippet b/Engine/Extras/VisualStudioSnippets/ue4del.snippet
new file mode 100644
--- /dev/null
+++ b/Engine/Extras/VisualStudioSnippets/ue4del.snippet
@@ -0,0 +1,28 @@
+<?xml version="1.0" encoding="utf-8"?>
+    <CodeSnippets xmlns="http://schemas.microsoft.com/VisualStudio/2005/CodeSnippet">
+      <CodeSnippet Format="1.0.0">
+        <Header>
+          <SnippetTypes>
+            <SnippetType>Expansion</SnippetType>
+          </SnippetTypes>
+          <Title>ue4del</Title>
+          <Author>Epic Games</Author>
+          <Description>Unreal Engine Delegate snippet</Description>
+          <HelpUrl>
+          </HelpUrl>
+          <Shortcut>ue4del</Shortcut>
+        </Header>
+        <Snippet>
+          <Declarations>
+       <Literal Editable="true">
+          <ID>delegateName</ID>
+          <ToolTip>Delegate type name</ToolTip>
+          <Default>DelegateName</Default>
+          <Function>
+          </Function>
+        </Literal>
+        </Declarations>
+          <Code Language="cpp"><![CDATA[DECLARE_DELEGATE($delegateName$)$selected$ $end$]]></Code>
+        </Snippet>
+      </CodeSnippet>
+    </CodeSnippets>
\ No newline at end of file
diff --git a/Engine/Extras/VisualStudioSnippets/ue4del1.snippet b/Engine/Extras/VisualStudioSnippets/ue4del1.snippet
new file mode 100644
--- /dev/null
+++ b/Engine/Extras/VisualStudioSnippets/ue4del1.snippet
@@ -0,0 +1,35 @@
+<?xml version="1.0" encoding="utf-8"?>
+    <CodeSnippets xmlns="http://schemas.microsoft.com/VisualStudio/2005/CodeSnippet">
+      <CodeSnippet Format="1.0.0">
+        <Header>
+          <SnippetTypes>
+            <SnippetType>Expansion</SnippetType>
+          </SnippetTypes>
+          <Title>ue4del1</Title>
+          <Author>Epic Games</Author>
+          <Description>Unreal Engine Delegate snippet</Description>
+          <HelpUrl>
+          </HelpUrl>
+          <Shortcut>ue4del1</Shortcut>
+        </Header>
+        <Snippet>
+          <Declarations>
+       <Literal Editable="true">
+          <ID>delegateName</ID>
+          <ToolTip>Delegate type name</ToolTip>
+          <Default>DelegateName</Default>
+          <Function>
+          </Function>
+        </Literal>
+       <Literal Editable="true">
+          <ID>paramType1</ID>
+          <ToolTip>Parameter type</ToolTip>
+          <Default>Param1Type</Default>
+          <Function>
+          </Function>
+            </Literal>
+        </Declarations>
+          <Code Language="cpp"><![CDATA[DECLARE_DELEGATE($delegateName$, $paramType1$)$selected$ $end$]]></Code>
+        </Snippet>
+      </CodeSnippet>
+    </CodeSnippets>
\ No newline at end of file
diff --git a/Engine/Extras/VisualStudioSnippets/ue4del2.snippet b/Engine/Extras/VisualStudioSnippets/ue4del2.snippet
new file mode 100644
--- /dev/null
+++ b/Engine/Extras/VisualStudioSnippets/ue4del2.snippet
@@ -0,0 +1,42 @@
+<?xml version="1.0" encoding="utf-8"?>
+    <CodeSnippets xmlns="http://schemas.microsoft.com/VisualStudio/2005/CodeSnippet">
+      <CodeSnippet Format="1.0.0">
+        <Header>
+          <SnippetTypes>
+            <SnippetType>Expansion</SnippetType>
+          </SnippetTypes>
+          <Title>ue4del2</Title>
+          <Author>Epic Games</Author>
+          <Description>Unreal Engine Delegate snippet</Description>
+          <HelpUrl>
+          </HelpUrl>
+          <Shortcut>ue4del2</Shortcut>
+        </Header>
+        <Snippet>
+          <Declarations>
+       <Literal Editable="true">
+          <ID>delegateName</ID>
+          <ToolTip>Delegate type name</ToolTip>
+          <Default>DelegateName</Default>
+          <Function>
+          </Function>
+        </Literal>
+       <Literal Editable="true">
+          <ID>paramType1</ID>
+          <ToolTip>Parameter type</ToolTip>
+          <Default>Param1Type</Default>
+          <Function>
+          </Function>
+            </Literal>
+       <Literal Editable="true">
+          <ID>paramType2</ID>
+          <ToolTip>Parameter type</ToolTip>
+          <Default>Param2Type</Default>
+          <Function>
+          </Function>
+            </Literal>
+        </Declarations>
+          <Code Language="cpp"><![CDATA[DECLARE_DELEGATE($delegateName$, $paramType1$, $paramType2$)$selected$ $end$]]></Code>
+        </Snippet>
+      </CodeSnippet>
+    </CodeSnippets>
\ No newline at end of file
diff --git a/Engine/Extras/VisualStudioSnippets/ue4del3.snippet b/Engine/Extras/VisualStudioSnippets/ue4del3.snippet
new file mode 100644
--- /dev/null
+++ b/Engine/Extras/VisualStudioSnippets/ue4del3.snippet
@@ -0,0 +1,49 @@
+<?xml version="1.0" encoding="utf-8"?>
+    <CodeSnippets xmlns="http://schemas.microsoft.com/VisualStudio/2005/CodeSnippet">
+      <CodeSnippet Format="1.0.0">
+        <Header>
+          <SnippetTypes>
+            <SnippetType>Expansion</SnippetType>
+          </SnippetTypes>
+          <Title>ue4del3</Title>
+          <Author>Epic Games</Author>
+          <Description>Unreal Engine Delegate snippet</Description>
+          <HelpUrl>
+          </HelpUrl>
+          <Shortcut>ue4del3</Shortcut>
+        </Header>
+        <Snippet>
+          <Declarations>
+       <Literal Editable="true">
+          <ID>delegateName</ID>
+          <ToolTip>Delegate type name</ToolTip>
+          <Default>DelegateName</Default>
+          <Function>
+          </Function>
+        </Literal>
+       <Literal Editable="true">
+          <ID>paramType1</ID>
+          <ToolTip>Parameter type</ToolTip>
+          <Default>Param1Type</Default>
+          <Function>
+          </Function>
+            </Literal>
+       <Literal Editable="true">
+          <ID>paramType2</ID>
+          <ToolTip>Parameter type</ToolTip>
+          <Default>Param2Type</Default>
+          <Function>
+          </Function>
+            </Literal>
+       <Literal Editable="true">
+          <ID>paramType3</ID>
+          <ToolTip>Parameter type</ToolTip>
+          <Default>Param3Type</Default>
+          <Function>
+          </Function>
+            </Literal>
+        </Declarations>
+          <Code Language="cpp"><![CDATA[DECLARE_DELEGATE($delegateName$, $paramType1$, $paramType2$, $paramType3$)$selected$ $end$]]></Code>
+        </Snippet>
+      </CodeSnippet>
+    </CodeSnippets>
\ No newline at end of file
diff --git a/Engine/Extras/VisualStudioSnippets/ue4del4.snippet b/Engine/Extras/VisualStudioSnippets/ue4del4.snippet
new file mode 100644
--- /dev/null
+++ b/Engine/Extras/VisualStudioSnippets/ue4del4.snippet
@@ -0,0 +1,56 @@
+<?xml version="1.0" encoding="utf-8"?>
+    <CodeSnippets xmlns="http://schemas.microsoft.com/VisualStudio/2005/CodeSnippet">
+      <CodeSnippet Format="1.0.0">
+        <Header>
+          <SnippetTypes>
+            <SnippetType>Expansion</SnippetType>
+          </SnippetTypes>
+          <Title>ue4del4</Title>
+          <Author>Epic Games</Author>
+          <Description>Unreal Engine Delegate snippet</Description>
+          <HelpUrl>
+          </HelpUrl>
+          <Shortcut>ue4del4</Shortcut>
+        </Header>
+        <Snippet>
+          <Declarations>
+       <Literal Editable="true">
+          <ID>delegateName</ID>
+          <ToolTip>Delegate type name</ToolTip>
+          <Default>DelegateName</Default>
+          <Function>
+          </Function>
+        </Literal>
+       <Literal Editable="true">
+          <ID>paramType1</ID>
+          <ToolTip>Parameter type</ToolTip>
+          <Default>Param1Type</Default>
+          <Function>
+          </Function>
+            </Literal>
+       <Literal Editable="true">
+          <ID>paramType2</ID>
+          <ToolTip>Parameter type</ToolTip>
+          <Default>Param2Type</Default>
+          <Function>
+          </Function>
+            </Literal>
+       <Literal Editable="true">
+          <ID>paramType3</ID>
+          <ToolTip>Parameter type</ToolTip>
+          <Default>Param3Type</Default>
+          <Function>
+          </Function>
+            </Literal>
+       <Literal Editable="true">
+          <ID>paramType4</ID>
+          <ToolTip>Parameter type</ToolTip>
+          <Default>Param4Type</Default>
+          <Function>
+          </Function>
+            </Literal>
+        </Declarations>
+          <Code Language="cpp"><![CDATA[DECLARE_DELEGATE($delegateName$, $paramType1$, $paramType2$, $paramType3$, $paramType4$)$selected$ $end$]]></Code>
+        </Snippet>
+      </CodeSnippet>
+    </CodeSnippets>
\ No newline at end of file
diff --git a/Engine/Extras/VisualStudioSnippets/ue4del5.snippet b/Engine/Extras/VisualStudioSnippets/ue4del5.snippet
new file mode 100644
--- /dev/null
+++ b/Engine/Extras/VisualStudioSnippets/ue4del5.snippet
@@ -0,0 +1,63 @@
+<?xml version="1.0" encoding="utf-8"?>
+    <CodeSnippets xmlns="http://schemas.microsoft.com/VisualStudio/2005/CodeSnippet">
+      <CodeSnippet Format="1.0.0">
+        <Header>
+          <SnippetTypes>
+            <SnippetType>Expansion</SnippetType>
+          </SnippetTypes>
+          <Title>ue4del5</Title>
+          <Author>Epic Games</Author>
+          <Description>Unreal Engine Delegate snippet</Description>
+          <HelpUrl>
+          </HelpUrl>
+          <Shortcut>ue4del5</Shortcut>
+        </Header>
+        <Snippet>
+          <Declarations>
+       <Literal Editable="true">
+          <ID>delegateName</ID>
+          <ToolTip>Delegate type name</ToolTip>
+          <Default>DelegateName</Default>
+          <Function>
+          </Function>
+        </Literal>
+       <Literal Editable="true">
+          <ID>paramType1</ID>
+          <ToolTip>Parameter type</ToolTip>
+          <Default>Param1Type</Default>
+          <Function>
+          </Function>
+            </Literal>
+       <Literal Editable="true">
+          <ID>paramType2</ID>
+          <ToolTip>Parameter type</ToolTip>
+          <Default>Param2Type</Default>
+          <Function>
+          </Function>
+            </Literal>
+       <Literal Editable="true">
+          <ID>paramType3</ID>
+          <ToolTip>Parameter type</ToolTip>
+          <Default>Param3Type</Default>
+          <Function>
+          </Function>
+            </Literal>
+       <Literal Editable="true">
+          <ID>paramType4</ID>
+          <ToolTip>Parameter type</ToolTip>
+          <Default>Param4Type</Default>
+          <Function>
+          </Function>
+            </Literal>
+       <Literal Editable="true">
+          <ID>paramType5</ID>
+          <ToolTip>Parameter type</ToolTip>
+          <Default>Param5Type</Default>
+          <Function>
+          </Function>
+            </Literal>
+        </Declarations>
+          <Code Language="cpp"><![CDATA[DECLARE_DELEGATE($delegateName$, $paramType1$, $paramType2$, $paramType3$, $paramType4$, $paramType5$)$selected$ $end$]]></Code>
+        </Snippet>
+      </CodeSnippet>
+    </CodeSnippets>
\ No newline at end of file
diff --git a/Engine/Extras/VisualStudioSnippets/ue4del6.snippet b/Engine/Extras/VisualStudioSnippets/ue4del6.snippet
new file mode 100644
--- /dev/null
+++ b/Engine/Extras/VisualStudioSnippets/ue4del6.snippet
@@ -0,0 +1,70 @@
+<?xml version="1.0" encoding="utf-8"?>
+    <CodeSnippets xmlns="http://schemas.microsoft.com/VisualStudio/2005/CodeSnippet">
+      <CodeSnippet Format="1.0.0">
+        <Header>
+          <SnippetTypes>
+            <SnippetType>Expansion</SnippetType>
+          </SnippetTypes>
+          <Title>ue4del6</Title>
+          <Author>Epic Games</Author>
+          <Description>Unreal Engine Delegate snippet</Description>
+          <HelpUrl>
+          </HelpUrl>
+          <Shortcut>ue4del6</Shortcut>
+        </Header>
+        <Snippet>
+          <Declarations>
+       <Literal Editable="true">
+          <ID>delegateName</ID>
+          <ToolTip>Delegate type name</ToolTip>
+          <Default>DelegateName</Default>
+          <Function>
+          </Function>
+        </Literal>
+       <Literal Editable="true">
+          <ID>paramType1</ID>
+          <ToolTip>Parameter type</ToolTip>
+          <Default>Param1Type</Default>
+          <Function>
+          </Function>
+            </Literal>
+       <Literal Editable="true">
+          <ID>paramType2</ID>
+          <ToolTip>Parameter type</ToolTip>
+          <Default>Param2Type</Default>
+          <Function>
+          </Function>
+            </Literal>
+       <Literal Editable="true">
+          <ID>paramType3</ID>
+          <ToolTip>Parameter type</ToolTip>
+          <Default>Param3Type</Default>
+          <Function>
+          </Function>
+            </Literal>
+       <Literal Editable="true">
+          <ID>paramType4</ID>
+          <ToolTip>Parameter type</ToolTip>
+          <Default>Param4Type</Default>
+          <Function>
+          </Function>
+            </Literal>
+       <Literal Editable="true">
+          <ID>paramType5</ID>
+          <ToolTip>Parameter type</ToolTip>
+          <Default>Param5Type</Default>
+          <Function>
+          </Function>
+            </Literal>
+       <Literal Editable="true">
+          <ID>paramType6</ID>
+          <ToolTip>Parameter type</ToolTip>
+          <Default>Param6Type</Default>
+          <Function>
+          </Function>
+            </Literal>
+        </Declarations>
+          <Code Language="cpp"><![CDATA[DECLARE_DELEGATE($delegateName$, $paramType1$, $paramType2$, $paramType3$, $paramType4$, $paramType5$, $paramType6$)$selected$ $end$]]></Code>
+        </Snippet>
+      </CodeSnippet>
+    </CodeSnippets>
\ No newline at end of file
diff --git a/Engine/Extras/VisualStudioSnippets/ue4del7.snippet b/Engine/Extras/VisualStudioSnippets/ue4del7.snippet
new file mode 100644
--- /dev/null
+++ b/Engine/Extras/VisualStudioSnippets/ue4del7.snippet
@@ -0,0 +1,77 @@
+<?xml version="1.0" encoding="utf-8"?>
+    <CodeSnippets xmlns="http://schemas.microsoft.com/VisualStudio/2005/CodeSnippet">
+      <CodeSnippet Format="1.0.0">
+        <Header>
+          <SnippetTypes>
+            <SnippetType>Expansion</SnippetType>
+          </SnippetTypes>
+          <Title>ue4del7</Title>
+          <Author>Epic Games</Author>
+          <Description>Unreal Engine Delegate snippet</Description>
+          <HelpUrl>
+          </HelpUrl>
+          <Shortcut>ue4del7</Shortcut>
+        </Header>
+        <Snippet>
+          <Declarations>
+       <Literal Editable="true">
+          <ID>delegateName</ID>
+          <ToolTip>Delegate type name</ToolTip>
+          <Default>DelegateName</Default>
+          <Function>
+          </Function>
+        </Literal>
+       <Literal Editable="true">
+          <ID>paramType1</ID>
+          <ToolTip>Parameter type</ToolTip>
+          <Default>Param1Type</Default>
+          <Function>
+          </Function>
+            </Literal>
+       <Literal Editable="true">
+          <ID>paramType2</ID>
+          <ToolTip>Parameter type</ToolTip>
+          <Default>Param2Type</Default>
+          <Function>
+          </Function>
+            </Literal>
+       <Literal Editable="true">
+          <ID>paramType3</ID>
+          <ToolTip>Parameter type</ToolTip>
+          <Default>Param3Type</Default>
+          <Function>
+          </Function>
+            </Literal>
+       <Literal Editable="true">
+          <ID>paramType4</ID>
+          <ToolTip>Parameter type</ToolTip>
+          <Default>Param4Type</Default>
+          <Function>
+          </Function>
+            </Literal>
+       <Literal Editable="true">
+          <ID>paramType5</ID>
+          <ToolTip>Parameter type</ToolTip>
+          <Default>Param5Type</Default>
+          <Function>
+          </Function>
+            </Literal>
+       <Literal Editable="true">
+          <ID>paramType6</ID>
+          <ToolTip>Parameter type</ToolTip>
+          <Default>Param6Type</Default>
+          <Function>
+          </Function>
+            </Literal>
+       <Literal Editable="true">
+          <ID>paramType7</ID>
+          <ToolTip>Parameter type</ToolTip>
+          <Default>Param7Type</Default>
+          <Function>
+          </Function>
+            </Literal>
+        </Declarations>
+          <Code Language="cpp"><![CDATA[DECLARE_DELEGATE($delegateName$, $paramType1$, $paramType2$, $paramType3$, $paramType4$, $paramType5$, $paramType6$, $paramType7$)$selected$ $end$]]></Code>
+        </Snippet>
+      </CodeSnippet>
+    </CodeSnippets>
\ No newline at end of file
diff --git a/Engine/Extras/VisualStudioSnippets/ue4del8.snippet b/Engine/Extras/VisualStudioSnippets/ue4del8.snippet
new file mode 100644
--- /dev/null
+++ b/Engine/Extras/VisualStudioSnippets/ue4del8.snippet
@@ -0,0 +1,84 @@
+<?xml version="1.0" encoding="utf-8"?>
+    <CodeSnippets xmlns="http://schemas.microsoft.com/VisualStudio/2005/CodeSnippet">
+      <CodeSnippet Format="1.0.0">
+        <Header>
+          <SnippetTypes>
+            <SnippetType>Expansion</SnippetType>
+          </SnippetTypes>
+          <Title>ue4del8</Title>
+          <Author>Epic Games</Author>
+          <Description>Unreal Engine Delegate snippet</Description>
+          <HelpUrl>
+          </HelpUrl>
+          <Shortcut>ue4del8</Shortcut>
+        </Header>
+        <Snippet>
+          <Declarations>
+       <Literal Editable="true">
+          <ID>delegateName</ID>
+          <ToolTip>Delegate type name</ToolTip>
+          <Default>DelegateName</Default>
+          <Function>
+          </Function>
+        </Literal>
+       <Literal Editable="true">
+          <ID>paramType1</ID>
+          <ToolTip>Parameter type</ToolTip>
+          <Default>Param1Type</Default>
+          <Function>
+          </Function>
+            </Literal>
+       <Literal Editable="true">
+          <ID>paramType2</ID>
+          <ToolTip>Parameter type</ToolTip>
+          <Default>Param2Type</Default>
+          <Function>
+          </Function>
+            </Literal>
+       <Literal Editable="true">
+          <ID>paramType3</ID>
+          <ToolTip>Parameter type</ToolTip>
+          <Default>Param3Type</Default>
+          <Function>
+          </Function>
+            </Literal>
+       <Literal Editable="true">
+          <ID>paramType4</ID>
+          <ToolTip>Parameter type</ToolTip>
+          <Default>Param4Type</Default>
+          <Function>
+          </Function>
+            </Literal>
+       <Literal Editable="true">
+          <ID>paramType5</ID>
+          <ToolTip>Parameter type</ToolTip>
+          <Default>Param5Type</Default>
+          <Function>
+          </Function>
+            </Literal>
+       <Literal Editable="true">
+          <ID>paramType6</ID>
+          <ToolTip>Parameter type</ToolTip>
+          <Default>Param6Type</Default>
+          <Function>
+          </Function>
+            </Literal>
+       <Literal Editable="true">
+          <ID>paramType7</ID>
+          <ToolTip>Parameter type</ToolTip>
+          <Default>Param7Type</Default>
+          <Function>
+          </Function>
+            </Literal>
+       <Literal Editable="true">
+          <ID>paramType8</ID>
+          <ToolTip>Parameter type</ToolTip>
+          <Default>Param8Type</Default>
+          <Function>
+          </Function>
+            </Literal>
+        </Declarations>
+          <Code Language="cpp"><![CDATA[DECLARE_DELEGATE($delegateName$, $paramType1$, $paramType2$, $paramType3$, $paramType4$, $paramType5$, $paramType6$, $paramType7$, $paramType8$)$selected$ $end$]]></Code>
+        </Snippet>
+      </CodeSnippet>
+    </CodeSnippets>
\ No newline at end of file
diff --git a/Engine/Extras/VisualStudioSnippets/ue4deldyn.snippet b/Engine/Extras/VisualStudioSnippets/ue4deldyn.snippet
new file mode 100644
--- /dev/null
+++ b/Engine/Extras/VisualStudioSnippets/ue4deldyn.snippet
@@ -0,0 +1,28 @@
+<?xml version="1.0" encoding="utf-8"?>
+    <CodeSnippets xmlns="http://schemas.microsoft.com/VisualStudio/2005/CodeSnippet">
+      <CodeSnippet Format="1.0.0">
+        <Header>
+          <SnippetTypes>
+            <SnippetType>Expansion</SnippetType>
+          </SnippetTypes>
+          <Title>ue4deldyn</Title>
+          <Author>Epic Games</Author>
+          <Description>Unreal Engine Dynamic Delegate snippet</Description>
+          <HelpUrl>
+          </HelpUrl>
+          <Shortcut>ue4deldyn</Shortcut>
+        </Header>
+        <Snippet>
+          <Declarations>
+       <Literal Editable="true">
+          <ID>delegateName</ID>
+          <ToolTip>Delegate type name</ToolTip>
+          <Default>DelegateName</Default>
+          <Function>
+          </Function>
+        </Literal>
+        </Declarations>
+          <Code Language="cpp"><![CDATA[DECLARE_DYNAMIC_DELEGATE($delegateName$)$selected$ $end$]]></Code>
+        </Snippet>
+      </CodeSnippet>
+    </CodeSnippets>
\ No newline at end of file
diff --git a/Engine/Extras/VisualStudioSnippets/ue4deldyn1.snippet b/Engine/Extras/VisualStudioSnippets/ue4deldyn1.snippet
new file mode 100644
--- /dev/null
+++ b/Engine/Extras/VisualStudioSnippets/ue4deldyn1.snippet
@@ -0,0 +1,42 @@
+<?xml version="1.0" encoding="utf-8"?>
+    <CodeSnippets xmlns="http://schemas.microsoft.com/VisualStudio/2005/CodeSnippet">
+      <CodeSnippet Format="1.0.0">
+        <Header>
+          <SnippetTypes>
+            <SnippetType>Expansion</SnippetType>
+          </SnippetTypes>
+          <Title>ue4deldyn1</Title>
+          <Author>Epic Games</Author>
+          <Description>Unreal Engine Dynamic Delegate snippet</Description>
+          <HelpUrl>
+          </HelpUrl>
+          <Shortcut>ue4deldyn1</Shortcut>
+        </Header>
+        <Snippet>
+          <Declarations>
+       <Literal Editable="true">
+          <ID>delegateName</ID>
+          <ToolTip>Delegate type name</ToolTip>
+          <Default>DelegateName</Default>
+          <Function>
+          </Function>
+        </Literal>
+        <Literal Editable="true">
+          <ID>paramName1</ID>
+          <ToolTip>Parameter name</ToolTip>
+          <Default>Param1Name</Default>
+          <Function>
+          </Function>
+            </Literal>
+       <Literal Editable="true">
+          <ID>paramType1</ID>
+          <ToolTip>Parameter type</ToolTip>
+          <Default>Param1Type</Default>
+          <Function>
+          </Function>
+            </Literal>
+        </Declarations>
+          <Code Language="cpp"><![CDATA[DECLARE_DYNAMIC_DELEGATE_OneParam($delegateName$, $paramType1$, $paramName1$)$selected$ $end$]]></Code>
+        </Snippet>
+      </CodeSnippet>
+    </CodeSnippets>
\ No newline at end of file
diff --git a/Engine/Extras/VisualStudioSnippets/ue4deldyn2.snippet b/Engine/Extras/VisualStudioSnippets/ue4deldyn2.snippet
new file mode 100644
--- /dev/null
+++ b/Engine/Extras/VisualStudioSnippets/ue4deldyn2.snippet
@@ -0,0 +1,56 @@
+<?xml version="1.0" encoding="utf-8"?>
+    <CodeSnippets xmlns="http://schemas.microsoft.com/VisualStudio/2005/CodeSnippet">
+      <CodeSnippet Format="1.0.0">
+        <Header>
+          <SnippetTypes>
+            <SnippetType>Expansion</SnippetType>
+          </SnippetTypes>
+          <Title>ue4deldyn2</Title>
+          <Author>Epic Games</Author>
+          <Description>Unreal Engine Dynamic Delegate snippet</Description>
+          <HelpUrl>
+          </HelpUrl>
+          <Shortcut>ue4deldyn2</Shortcut>
+        </Header>
+        <Snippet>
+          <Declarations>
+       <Literal Editable="true">
+          <ID>delegateName</ID>
+          <ToolTip>Delegate type name</ToolTip>
+          <Default>DelegateName</Default>
+          <Function>
+          </Function>
+        </Literal>
+        <Literal Editable="true">
+          <ID>paramName1</ID>
+          <ToolTip>Parameter name</ToolTip>
+          <Default>Param1Name</Default>
+          <Function>
+          </Function>
+            </Literal>
+        <Literal Editable="true">
+          <ID>paramName2</ID>
+          <ToolTip>Parameter name</ToolTip>
+          <Default>Param2Name</Default>
+          <Function>
+          </Function>
+            </Literal>
+       <Literal Editable="true">
+          <ID>paramType1</ID>
+          <ToolTip>Parameter type</ToolTip>
+          <Default>Param1Type</Default>
+          <Function>
+          </Function>
+            </Literal>
+       <Literal Editable="true">
+          <ID>paramType2</ID>
+          <ToolTip>Parameter type</ToolTip>
+          <Default>Param2Type</Default>
+          <Function>
+          </Function>
+            </Literal>
+        </Declarations>
+          <Code Language="cpp"><![CDATA[DECLARE_DYNAMIC_DELEGATE_TwoParams($delegateName$, $paramType1$, $paramName1$, $paramType2$, $paramName2$)$selected$ $end$]]></Code>
+        </Snippet>
+      </CodeSnippet>
+    </CodeSnippets>
\ No newline at end of file
diff --git a/Engine/Extras/VisualStudioSnippets/ue4deldyn3.snippet b/Engine/Extras/VisualStudioSnippets/ue4deldyn3.snippet
new file mode 100644
--- /dev/null
+++ b/Engine/Extras/VisualStudioSnippets/ue4deldyn3.snippet
@@ -0,0 +1,70 @@
+<?xml version="1.0" encoding="utf-8"?>
+    <CodeSnippets xmlns="http://schemas.microsoft.com/VisualStudio/2005/CodeSnippet">
+      <CodeSnippet Format="1.0.0">
+        <Header>
+          <SnippetTypes>
+            <SnippetType>Expansion</SnippetType>
+          </SnippetTypes>
+          <Title>ue4deldyn3</Title>
+          <Author>Epic Games</Author>
+          <Description>Unreal Engine Dynamic Delegate snippet</Description>
+          <HelpUrl>
+          </HelpUrl>
+          <Shortcut>ue4deldyn3</Shortcut>
+        </Header>
+        <Snippet>
+          <Declarations>
+       <Literal Editable="true">
+          <ID>delegateName</ID>
+          <ToolTip>Delegate type name</ToolTip>
+          <Default>DelegateName</Default>
+          <Function>
+          </Function>
+        </Literal>
+        <Literal Editable="true">
+          <ID>paramName1</ID>
+          <ToolTip>Parameter name</ToolTip>
+          <Default>Param1Name</Default>
+          <Function>
+          </Function>
+            </Literal>
+        <Literal Editable="true">
+          <ID>paramName2</ID>
+          <ToolTip>Parameter name</ToolTip>
+          <Default>Param2Name</Default>
+          <Function>
+          </Function>
+            </Literal>
+        <Literal Editable="true">
+          <ID>paramName3</ID>
+          <ToolTip>Parameter name</ToolTip>
+          <Default>Param3Name</Default>
+          <Function>
+          </Function>
+            </Literal>
+       <Literal Editable="true">
+          <ID>paramType1</ID>
+          <ToolTip>Parameter type</ToolTip>
+          <Default>Param1Type</Default>
+          <Function>
+          </Function>
+            </Literal>
+       <Literal Editable="true">
+          <ID>paramType2</ID>
+          <ToolTip>Parameter type</ToolTip>
+          <Default>Param2Type</Default>
+          <Function>
+          </Function>
+            </Literal>
+       <Literal Editable="true">
+          <ID>paramType3</ID>
+          <ToolTip>Parameter type</ToolTip>
+          <Default>Param3Type</Default>
+          <Function>
+          </Function>
+            </Literal>
+        </Declarations>
+          <Code Language="cpp"><![CDATA[DECLARE_DYNAMIC_DELEGATE_ThreeParams($delegateName$, $paramType1$, $paramName1$, $paramType2$, $paramName2$, $paramType3$, $paramName3$)$selected$ $end$]]></Code>
+        </Snippet>
+      </CodeSnippet>
+    </CodeSnippets>
\ No newline at end of file
diff --git a/Engine/Extras/VisualStudioSnippets/ue4deldyn4.snippet b/Engine/Extras/VisualStudioSnippets/ue4deldyn4.snippet
new file mode 100644
--- /dev/null
+++ b/Engine/Extras/VisualStudioSnippets/ue4deldyn4.snippet
@@ -0,0 +1,84 @@
+<?xml version="1.0" encoding="utf-8"?>
+    <CodeSnippets xmlns="http://schemas.microsoft.com/VisualStudio/2005/CodeSnippet">
+      <CodeSnippet Format="1.0.0">
+        <Header>
+          <SnippetTypes>
+            <SnippetType>Expansion</SnippetType>
+          </SnippetTypes>
+          <Title>ue4deldyn4</Title>
+          <Author>Epic Games</Author>
+          <Description>Unreal Engine Dynamic Delegate snippet</Description>
+          <HelpUrl>
+          </HelpUrl>
+          <Shortcut>ue4deldyn4</Shortcut>
+        </Header>
+        <Snippet>
+          <Declarations>
+       <Literal Editable="true">
+          <ID>delegateName</ID>
+          <ToolTip>Delegate type name</ToolTip>
+          <Default>DelegateName</Default>
+          <Function>
+          </Function>
+        </Literal>
+        <Literal Editable="true">
+          <ID>paramName1</ID>
+          <ToolTip>Parameter name</ToolTip>
+          <Default>Param1Name</Default>
+          <Function>
+          </Function>
+            </Literal>
+        <Literal Editable="true">
+          <ID>paramName2</ID>
+          <ToolTip>Parameter name</ToolTip>
+          <Default>Param2Name</Default>
+          <Function>
+          </Function>
+            </Literal>
+        <Literal Editable="true">
+          <ID>paramName3</ID>
+          <ToolTip>Parameter name</ToolTip>
+          <Default>Param3Name</Default>
+          <Function>
+          </Function>
+            </Literal>
+        <Literal Editable="true">
+          <ID>paramName4</ID>
+          <ToolTip>Parameter name</ToolTip>
+          <Default>Param4Name</Default>
+          <Function>
+          </Function>
+            </Literal>
+       <Literal Editable="true">
+          <ID>paramType1</ID>
+          <ToolTip>Parameter type</ToolTip>
+          <Default>Param1Type</Default>
+          <Function>
+          </Function>
+            </Literal>
+       <Literal Editable="true">
+          <ID>paramType2</ID>
+          <ToolTip>Parameter type</ToolTip>
+          <Default>Param2Type</Default>
+          <Function>
+          </Function>
+            </Literal>
+       <Literal Editable="true">
+          <ID>paramType3</ID>
+          <ToolTip>Parameter type</ToolTip>
+          <Default>Param3Type</Default>
+          <Function>
+          </Function>
+            </Literal>
+       <Literal Editable="true">
+          <ID>paramType4</ID>
+          <ToolTip>Parameter type</ToolTip>
+          <Default>Param4Type</Default>
+          <Function>
+          </Function>
+            </Literal>
+        </Declarations>
+          <Code Language="cpp"><![CDATA[DECLARE_DYNAMIC_DELEGATE_FourParams($delegateName$, $paramType1$, $paramName1$, $paramType2$, $paramName2$, $paramType3$, $paramName3$, $paramType4$, $paramName4$)$selected$ $end$]]></Code>
+        </Snippet>
+      </CodeSnippet>
+    </CodeSnippets>
\ No newline at end of file
diff --git a/Engine/Extras/VisualStudioSnippets/ue4deldyn5.snippet b/Engine/Extras/VisualStudioSnippets/ue4deldyn5.snippet
new file mode 100644
--- /dev/null
+++ b/Engine/Extras/VisualStudioSnippets/ue4deldyn5.snippet
@@ -0,0 +1,98 @@
+<?xml version="1.0" encoding="utf-8"?>
+    <CodeSnippets xmlns="http://schemas.microsoft.com/VisualStudio/2005/CodeSnippet">
+      <CodeSnippet Format="1.0.0">
+        <Header>
+          <SnippetTypes>
+            <SnippetType>Expansion</SnippetType>
+          </SnippetTypes>
+          <Title>ue4deldyn5</Title>
+          <Author>Epic Games</Author>
+          <Description>Unreal Engine Dynamic Delegate snippet</Description>
+          <HelpUrl>
+          </HelpUrl>
+          <Shortcut>ue4deldyn5</Shortcut>
+        </Header>
+        <Snippet>
+          <Declarations>
+       <Literal Editable="true">
+          <ID>delegateName</ID>
+          <ToolTip>Delegate type name</ToolTip>
+          <Default>DelegateName</Default>
+          <Function>
+          </Function>
+        </Literal>
+        <Literal Editable="true">
+          <ID>paramName1</ID>
+          <ToolTip>Parameter name</ToolTip>
+          <Default>Param1Name</Default>
+          <Function>
+          </Function>
+            </Literal>
+        <Literal Editable="true">
+          <ID>paramName2</ID>
+          <ToolTip>Parameter name</ToolTip>
+          <Default>Param2Name</Default>
+          <Function>
+          </Function>
+            </Literal>
+        <Literal Editable="true">
+          <ID>paramName3</ID>
+          <ToolTip>Parameter name</ToolTip>
+          <Default>Param3Name</Default>
+          <Function>
+          </Function>
+            </Literal>
+        <Literal Editable="true">
+          <ID>paramName4</ID>
+          <ToolTip>Parameter name</ToolTip>
+          <Default>Param4Name</Default>
+          <Function>
+          </Function>
+            </Literal>
+        <Literal Editable="true">
+          <ID>paramName5</ID>
+          <ToolTip>Parameter name</ToolTip>
+          <Default>Param5Name</Default>
+          <Function>
+          </Function>
+            </Literal>
+       <Literal Editable="true">
+          <ID>paramType1</ID>
+          <ToolTip>Parameter type</ToolTip>
+          <Default>Param1Type</Default>
+          <Function>
+          </Function>
+            </Literal>
+       <Literal Editable="true">
+          <ID>paramType2</ID>
+          <ToolTip>Parameter type</ToolTip>
+          <Default>Param2Type</Default>
+          <Function>
+          </Function>
+            </Literal>
+       <Literal Editable="true">
+          <ID>paramType3</ID>
+          <ToolTip>Parameter type</ToolTip>
+          <Default>Param3Type</Default>
+          <Function>
+          </Function>
+            </Literal>
+       <Literal Editable="true">
+          <ID>paramType4</ID>
+          <ToolTip>Parameter type</ToolTip>
+          <Default>Param4Type</Default>
+          <Function>
+          </Function>
+            </Literal>
+       <Literal Editable="true">
+          <ID>paramType5</ID>
+          <ToolTip>Parameter type</ToolTip>
+          <Default>Param5Type</Default>
+          <Function>
+          </Function>
+            </Literal>
+        </Declarations>
+          <Code Language="cpp"><![CDATA[DECLARE_DYNAMIC_DELEGATE_FiveParams($delegateName$, $paramType1$, $paramName1$, $paramType2$, $paramName2$, $paramType3$, $paramName3$, $paramType4$, $paramName4$, $paramType5$, $paramName5$)$selected$ $end$]]></Code>
+        </Snippet>
+      </CodeSnippet>
+    </CodeSnippets>
\ No newline at end of file
diff --git a/Engine/Extras/VisualStudioSnippets/ue4deldyn6.snippet b/Engine/Extras/VisualStudioSnippets/ue4deldyn6.snippet
new file mode 100644
--- /dev/null
+++ b/Engine/Extras/VisualStudioSnippets/ue4deldyn6.snippet
@@ -0,0 +1,112 @@
+<?xml version="1.0" encoding="utf-8"?>
+    <CodeSnippets xmlns="http://schemas.microsoft.com/VisualStudio/2005/CodeSnippet">
+      <CodeSnippet Format="1.0.0">
+        <Header>
+          <SnippetTypes>
+            <SnippetType>Expansion</SnippetType>
+          </SnippetTypes>
+          <Title>ue4deldyn6</Title>
+          <Author>Epic Games</Author>
+          <Description>Unreal Engine Dynamic Delegate snippet</Description>
+          <HelpUrl>
+          </HelpUrl>
+          <Shortcut>ue4deldyn6</Shortcut>
+        </Header>
+        <Snippet>
+          <Declarations>
+       <Literal Editable="true">
+          <ID>delegateName</ID>
+          <ToolTip>Delegate type name</ToolTip>
+          <Default>DelegateName</Default>
+          <Function>
+          </Function>
+        </Literal>
+        <Literal Editable="true">
+          <ID>paramName1</ID>
+          <ToolTip>Parameter name</ToolTip>
+          <Default>Param1Name</Default>
+          <Function>
+          </Function>
+            </Literal>
+        <Literal Editable="true">
+          <ID>paramName2</ID>
+          <ToolTip>Parameter name</ToolTip>
+          <Default>Param2Name</Default>
+          <Function>
+          </Function>
+            </Literal>
+        <Literal Editable="true">
+          <ID>paramName3</ID>
+          <ToolTip>Parameter name</ToolTip>
+          <Default>Param3Name</Default>
+          <Function>
+          </Function>
+            </Literal>
+        <Literal Editable="true">
+          <ID>paramName4</ID>
+          <ToolTip>Parameter name</ToolTip>
+          <Default>Param4Name</Default>
+          <Function>
+          </Function>
+            </Literal>
+        <Literal Editable="true">
+          <ID>paramName5</ID>
+          <ToolTip>Parameter name</ToolTip>
+          <Default>Param5Name</Default>
+          <Function>
+          </Function>
+            </Literal>
+        <Literal Editable="true">
+          <ID>paramName6</ID>
+          <ToolTip>Parameter name</ToolTip>
+          <Default>Param6Name</Default>
+          <Function>
+          </Function>
+            </Literal>
+       <Literal Editable="true">
+          <ID>paramType1</ID>
+          <ToolTip>Parameter type</ToolTip>
+          <Default>Param1Type</Default>
+          <Function>
+          </Function>
+            </Literal>
+       <Literal Editable="true">
+          <ID>paramType2</ID>
+          <ToolTip>Parameter type</ToolTip>
+          <Default>Param2Type</Default>
+          <Function>
+          </Function>
+            </Literal>
+       <Literal Editable="true">
+          <ID>paramType3</ID>
+          <ToolTip>Parameter type</ToolTip>
+          <Default>Param3Type</Default>
+          <Function>
+          </Function>
+            </Literal>
+       <Literal Editable="true">
+          <ID>paramType4</ID>
+          <ToolTip>Parameter type</ToolTip>
+          <Default>Param4Type</Default>
+          <Function>
+          </Function>
+            </Literal>
+       <Literal Editable="true">
+          <ID>paramType5</ID>
+          <ToolTip>Parameter type</ToolTip>
+          <Default>Param5Type</Default>
+          <Function>
+          </Function>
+            </Literal>
+       <Literal Editable="true">
+          <ID>paramType6</ID>
+          <ToolTip>Parameter type</ToolTip>
+          <Default>Param6Type</Default>
+          <Function>
+          </Function>
+            </Literal>
+        </Declarations>
+          <Code Language="cpp"><![CDATA[DECLARE_DYNAMIC_DELEGATE_SixParams($delegateName$, $paramType1$, $paramName1$, $paramType2$, $paramName2$, $paramType3$, $paramName3$, $paramType4$, $paramName4$, $paramType5$, $paramName5$, $paramType6$, $paramName6$)$selected$ $end$]]></Code>
+        </Snippet>
+      </CodeSnippet>
+    </CodeSnippets>
\ No newline at end of file
diff --git a/Engine/Extras/VisualStudioSnippets/ue4deldyn7.snippet b/Engine/Extras/VisualStudioSnippets/ue4deldyn7.snippet
new file mode 100644
--- /dev/null
+++ b/Engine/Extras/VisualStudioSnippets/ue4deldyn7.snippet
@@ -0,0 +1,126 @@
+<?xml version="1.0" encoding="utf-8"?>
+    <CodeSnippets xmlns="http://schemas.microsoft.com/VisualStudio/2005/CodeSnippet">
+      <CodeSnippet Format="1.0.0">
+        <Header>
+          <SnippetTypes>
+            <SnippetType>Expansion</SnippetType>
+          </SnippetTypes>
+          <Title>ue4deldyn7</Title>
+          <Author>Epic Games</Author>
+          <Description>Unreal Engine Dynamic Delegate snippet</Description>
+          <HelpUrl>
+          </HelpUrl>
+          <Shortcut>ue4deldyn7</Shortcut>
+        </Header>
+        <Snippet>
+          <Declarations>
+       <Literal Editable="true">
+          <ID>delegateName</ID>
+          <ToolTip>Delegate type name</ToolTip>
+          <Default>DelegateName</Default>
+          <Function>
+          </Function>
+        </Literal>
+        <Literal Editable="true">
+          <ID>paramName1</ID>
+          <ToolTip>Parameter name</ToolTip>
+          <Default>Param1Name</Default>
+          <Function>
+          </Function>
+            </Literal>
+        <Literal Editable="true">
+          <ID>paramName2</ID>
+          <ToolTip>Parameter name</ToolTip>
+          <Default>Param2Name</Default>
+          <Function>
+          </Function>
+            </Literal>
+        <Literal Editable="true">
+          <ID>paramName3</ID>
+          <ToolTip>Parameter name</ToolTip>
+          <Default>Param3Name</Default>
+          <Function>
+          </Function>
+            </Literal>
+        <Literal Editable="true">
+          <ID>paramName4</ID>
+          <ToolTip>Parameter name</ToolTip>
+          <Default>Param4Name</Default>
+          <Function>
+          </Function>
+            </Literal>
+        <Literal Editable="true">
+          <ID>paramName5</ID>
+          <ToolTip>Parameter name</ToolTip>
+          <Default>Param5Name</Default>
+          <Function>
+          </Function>
+            </Literal>
+        <Literal Editable="true">
+          <ID>paramName6</ID>
+          <ToolTip>Parameter name</ToolTip>
+          <Default>Param6Name</Default>
+          <Function>
+          </Function>
+            </Literal>
+        <Literal Editable="true">
+          <ID>paramName7</ID>
+          <ToolTip>Parameter name</ToolTip>
+          <Default>Param7Name</Default>
+          <Function>
+          </Function>
+            </Literal>
+       <Literal Editable="true">
+          <ID>paramType1</ID>
+          <ToolTip>Parameter type</ToolTip>
+          <Default>Param1Type</Default>
+          <Function>
+          </Function>
+            </Literal>
+       <Literal Editable="true">
+          <ID>paramType2</ID>
+          <ToolTip>Parameter type</ToolTip>
+          <Default>Param2Type</Default>
+          <Function>
+          </Function>
+            </Literal>
+       <Literal Editable="true">
+          <ID>paramType3</ID>
+          <ToolTip>Parameter type</ToolTip>
+          <Default>Param3Type</Default>
+          <Function>
+          </Function>
+            </Literal>
+       <Literal Editable="true">
+          <ID>paramType4</ID>
+          <ToolTip>Parameter type</ToolTip>
+          <Default>Param4Type</Default>
+          <Function>
+          </Function>
+            </Literal>
+       <Literal Editable="true">
+          <ID>paramType5</ID>
+          <ToolTip>Parameter type</ToolTip>
+          <Default>Param5Type</Default>
+          <Function>
+          </Function>
+            </Literal>
+       <Literal Editable="true">
+          <ID>paramType6</ID>
+          <ToolTip>Parameter type</ToolTip>
+          <Default>Param6Type</Default>
+          <Function>
+          </Function>
+            </Literal>
+       <Literal Editable="true">
+          <ID>paramType7</ID>
+          <ToolTip>Parameter type</ToolTip>
+          <Default>Param7Type</Default>
+          <Function>
+          </Function>
+            </Literal>
+        </Declarations>
+          <Code Language="cpp"><![CDATA[DECLARE_DYNAMIC_DELEGATE_SevenParams($delegateName$, $paramType1$, $paramName1$, $paramType2$, $paramName2$, $paramType3$, $paramName3$, $paramType4$, $paramName4$, $paramType5$, $paramName5$, $paramType6$, $paramName6$, $paramType7$, $paramName7$)$selected$ $end$]]></Code>
+        </Snippet>
+      </CodeSnippet>
+    </CodeSnippets>
\ No newline at end of file
diff --git a/Engine/Extras/VisualStudioSnippets/ue4deldyn8.snippet b/Engine/Extras/VisualStudioSnippets/ue4deldyn8.snippet
new file mode 100644
--- /dev/null
+++ b/Engine/Extras/VisualStudioSnippets/ue4deldyn8.snippet
@@ -0,0 +1,140 @@
+<?xml version="1.0" encoding="utf-8"?>
+    <CodeSnippets xmlns="http://schemas.microsoft.com/VisualStudio/2005/CodeSnippet">
+      <CodeSnippet Format="1.0.0">
+        <Header>
+          <SnippetTypes>
+            <SnippetType>Expansion</SnippetType>
+          </SnippetTypes>
+          <Title>ue4deldyn8</Title>
+          <Author>Epic Games</Author>
+          <Description>Unreal Engine Dynamic Delegate snippet</Description>
+          <HelpUrl>
+          </HelpUrl>
+          <Shortcut>ue4deldyn8</Shortcut>
+        </Header>
+        <Snippet>
+          <Declarations>
+       <Literal Editable="true">
+          <ID>delegateName</ID>
+          <ToolTip>Delegate type name</ToolTip>
+          <Default>DelegateName</Default>
+          <Function>
+          </Function>
+        </Literal>
+        <Literal Editable="true">
+          <ID>paramName1</ID>
+          <ToolTip>Parameter name</ToolTip>
+          <Default>Param1Name</Default>
+          <Function>
+          </Function>
+            </Literal>
+        <Literal Editable="true">
+          <ID>paramName2</ID>
+          <ToolTip>Parameter name</ToolTip>
+          <Default>Param2Name</Default>
+          <Function>
+          </Function>
+            </Literal>
+        <Literal Editable="true">
+          <ID>paramName3</ID>
+          <ToolTip>Parameter name</ToolTip>
+          <Default>Param3Name</Default>
+          <Function>
+          </Function>
+            </Literal>
+        <Literal Editable="true">
+          <ID>paramName4</ID>
+          <ToolTip>Parameter name</ToolTip>
+          <Default>Param4Name</Default>
+          <Function>
+          </Function>
+            </Literal>
+        <Literal Editable="true">
+          <ID>paramName5</ID>
+          <ToolTip>Parameter name</ToolTip>
+          <Default>Param5Name</Default>
+          <Function>
+          </Function>
+            </Literal>
+        <Literal Editable="true">
+          <ID>paramName6</ID>
+          <ToolTip>Parameter name</ToolTip>
+          <Default>Param6Name</Default>
+          <Function>
+          </Function>
+            </Literal>
+        <Literal Editable="true">
+          <ID>paramName7</ID>
+          <ToolTip>Parameter name</ToolTip>
+          <Default>Param7Name</Default>
+          <Function>
+          </Function>
+            </Literal>
+        <Literal Editable="true">
+          <ID>paramName8</ID>
+          <ToolTip>Parameter name</ToolTip>
+          <Default>Param8Name</Default>
+          <Function>
+          </Function>
+            </Literal>
+       <Literal Editable="true">
+          <ID>paramType1</ID>
+          <ToolTip>Parameter type</ToolTip>
+          <Default>Param1Type</Default>
+          <Function>
+          </Function>
+            </Literal>
+       <Literal Editable="true">
+          <ID>paramType2</ID>
+          <ToolTip>Parameter type</ToolTip>
+          <Default>Param2Type</Default>
+          <Function>
+          </Function>
+            </Literal>
+       <Literal Editable="true">
+          <ID>paramType3</ID>
+          <ToolTip>Parameter type</ToolTip>
+          <Default>Param3Type</Default>
+          <Function>
+          </Function>
+            </Literal>
+       <Literal Editable="true">
+          <ID>paramType4</ID>
+          <ToolTip>Parameter type</ToolTip>
+          <Default>Param4Type</Default>
+          <Function>
+          </Function>
+            </Literal>
+       <Literal Editable="true">
+          <ID>paramType5</ID>
+          <ToolTip>Parameter type</ToolTip>
+          <Default>Param5Type</Default>
+          <Function>
+          </Function>
+            </Literal>
+       <Literal Editable="true">
+          <ID>paramType6</ID>
+          <ToolTip>Parameter type</ToolTip>
+          <Default>Param6Type</Default>
+          <Function>
+          </Function>
+            </Literal>
+       <Literal Editable="true">
+          <ID>paramType7</ID>
+          <ToolTip>Parameter type</ToolTip>
+          <Default>Param7Type</Default>
+          <Function>
+          </Function>
+            </Literal>
+       <Literal Editable="true">
+          <ID>paramType8</ID>
+          <ToolTip>Parameter type</ToolTip>
+          <Default>Param8Type</Default>
+          <Function>
+          </Function>
+            </Literal>
+        </Declarations>
+          <Code Language="cpp"><![CDATA[DECLARE_DYNAMIC_DELEGATE_EightParams($delegateName$, $paramType1$, $paramName1$, $paramType2$, $paramName2$, $paramType3$, $paramName3$, $paramType4$, $paramName4$, $paramType5$, $paramName5$, $paramType6$, $paramName6$, $paramType7$, $paramName7$, $paramType8$, $paramName8$)$selected$ $end$]]></Code>
+        </Snippet>
+      </CodeSnippet>
+    </CodeSnippets>
\ No newline at end of file
diff --git a/Engine/Extras/VisualStudioSnippets/ue4deldynmul.snippet b/Engine/Extras/VisualStudioSnippets/ue4deldynmul.snippet
new file mode 100644
--- /dev/null
+++ b/Engine/Extras/VisualStudioSnippets/ue4deldynmul.snippet
@@ -0,0 +1,28 @@
+<?xml version="1.0" encoding="utf-8"?>
+    <CodeSnippets xmlns="http://schemas.microsoft.com/VisualStudio/2005/CodeSnippet">
+      <CodeSnippet Format="1.0.0">
+        <Header>
+          <SnippetTypes>
+            <SnippetType>Expansion</SnippetType>
+          </SnippetTypes>
+          <Title>ue4deldynmul</Title>
+          <Author>Epic Games</Author>
+          <Description>Unreal Engine Dynamic Multicast Delegate snippet</Description>
+          <HelpUrl>
+          </HelpUrl>
+          <Shortcut>ue4deldynmul</Shortcut>
+        </Header>
+        <Snippet>
+          <Declarations>
+       <Literal Editable="true">
+          <ID>delegateName</ID>
+          <ToolTip>Delegate type name</ToolTip>
+          <Default>DelegateName</Default>
+          <Function>
+          </Function>
+        </Literal>
+        </Declarations>
+          <Code Language="cpp"><![CDATA[DECLARE_DYNAMIC_MULTICAST_DELEGATE($delegateName$)$selected$ $end$]]></Code>
+        </Snippet>
+      </CodeSnippet>
+    </CodeSnippets>
\ No newline at end of file
diff --git a/Engine/Extras/VisualStudioSnippets/ue4deldynmul1.snippet b/Engine/Extras/VisualStudioSnippets/ue4deldynmul1.snippet
new file mode 100644
--- /dev/null
+++ b/Engine/Extras/VisualStudioSnippets/ue4deldynmul1.snippet
@@ -0,0 +1,42 @@
+<?xml version="1.0" encoding="utf-8"?>
+    <CodeSnippets xmlns="http://schemas.microsoft.com/VisualStudio/2005/CodeSnippet">
+      <CodeSnippet Format="1.0.0">
+        <Header>
+          <SnippetTypes>
+            <SnippetType>Expansion</SnippetType>
+          </SnippetTypes>
+          <Title>ue4deldynmul1</Title>
+          <Author>Epic Games</Author>
+          <Description>Unreal Engine Dynamic Multicast Delegate snippet</Description>
+          <HelpUrl>
+          </HelpUrl>
+          <Shortcut>ue4deldynmul1</Shortcut>
+        </Header>
+        <Snippet>
+          <Declarations>
+       <Literal Editable="true">
+          <ID>delegateName</ID>
+          <ToolTip>Delegate type name</ToolTip>
+          <Default>DelegateName</Default>
+          <Function>
+          </Function>
+        </Literal>
+        <Literal Editable="true">
+          <ID>paramName1</ID>
+          <ToolTip>Parameter name</ToolTip>
+          <Default>Param1Name</Default>
+          <Function>
+          </Function>
+            </Literal>
+       <Literal Editable="true">
+          <ID>paramType1</ID>
+          <ToolTip>Parameter type</ToolTip>
+          <Default>Param1Type</Default>
+          <Function>
+          </Function>
+            </Literal>
+        </Declarations>
+          <Code Language="cpp"><![CDATA[DECLARE_DYNAMIC_MULTICAST_DELEGATE_OneParam($delegateName$, $paramType1$, $paramName1$)$selected$ $end$]]></Code>
+        </Snippet>
+      </CodeSnippet>
+    </CodeSnippets>
\ No newline at end of file
diff --git a/Engine/Extras/VisualStudioSnippets/ue4deldynmul2.snippet b/Engine/Extras/VisualStudioSnippets/ue4deldynmul2.snippet
new file mode 100644
--- /dev/null
+++ b/Engine/Extras/VisualStudioSnippets/ue4deldynmul2.snippet
@@ -0,0 +1,56 @@
+<?xml version="1.0" encoding="utf-8"?>
+    <CodeSnippets xmlns="http://schemas.microsoft.com/VisualStudio/2005/CodeSnippet">
+      <CodeSnippet Format="1.0.0">
+        <Header>
+          <SnippetTypes>
+            <SnippetType>Expansion</SnippetType>
+          </SnippetTypes>
+          <Title>ue4deldynmul2</Title>
+          <Author>Epic Games</Author>
+          <Description>Unreal Engine Dynamic Multicast Delegate snippet</Description>
+          <HelpUrl>
+          </HelpUrl>
+          <Shortcut>ue4deldynmul2</Shortcut>
+        </Header>
+        <Snippet>
+          <Declarations>
+       <Literal Editable="true">
+          <ID>delegateName</ID>
+          <ToolTip>Delegate type name</ToolTip>
+          <Default>DelegateName</Default>
+          <Function>
+          </Function>
+        </Literal>
+        <Literal Editable="true">
+          <ID>paramName1</ID>
+          <ToolTip>Parameter name</ToolTip>
+          <Default>Param1Name</Default>
+          <Function>
+          </Function>
+            </Literal>
+        <Literal Editable="true">
+          <ID>paramName2</ID>
+          <ToolTip>Parameter name</ToolTip>
+          <Default>Param2Name</Default>
+          <Function>
+          </Function>
+            </Literal>
+       <Literal Editable="true">
+          <ID>paramType1</ID>
+          <ToolTip>Parameter type</ToolTip>
+          <Default>Param1Type</Default>
+          <Function>
+          </Function>
+            </Literal>
+       <Literal Editable="true">
+          <ID>paramType2</ID>
+          <ToolTip>Parameter type</ToolTip>
+          <Default>Param2Type</Default>
+          <Function>
+          </Function>
+            </Literal>
+        </Declarations>
+          <Code Language="cpp"><![CDATA[DECLARE_DYNAMIC_MULTICAST_DELEGATE_TwoParams($delegateName$, $paramType1$, $paramName1$, $paramType2$, $paramName2$)$selected$ $end$]]></Code>
+        </Snippet>
+      </CodeSnippet>
+    </CodeSnippets>
\ No newline at end of file
diff --git a/Engine/Extras/VisualStudioSnippets/ue4deldynmul3.snippet b/Engine/Extras/VisualStudioSnippets/ue4deldynmul3.snippet
new file mode 100644
--- /dev/null
+++ b/Engine/Extras/VisualStudioSnippets/ue4deldynmul3.snippet
@@ -0,0 +1,70 @@
+<?xml version="1.0" encoding="utf-8"?>
+    <CodeSnippets xmlns="http://schemas.microsoft.com/VisualStudio/2005/CodeSnippet">
+      <CodeSnippet Format="1.0.0">
+        <Header>
+          <SnippetTypes>
+            <SnippetType>Expansion</SnippetType>
+          </SnippetTypes>
+          <Title>ue4deldynmul3</Title>
+          <Author>Epic Games</Author>
+          <Description>Unreal Engine Dynamic Multicast Delegate snippet</Description>
+          <HelpUrl>
+          </HelpUrl>
+          <Shortcut>ue4deldynmul3</Shortcut>
+        </Header>
+        <Snippet>
+          <Declarations>
+       <Literal Editable="true">
+          <ID>delegateName</ID>
+          <ToolTip>Delegate type name</ToolTip>
+          <Default>DelegateName</Default>
+          <Function>
+          </Function>
+        </Literal>
+        <Literal Editable="true">
+          <ID>paramName1</ID>
+          <ToolTip>Parameter name</ToolTip>
+          <Default>Param1Name</Default>
+          <Function>
+          </Function>
+            </Literal>
+        <Literal Editable="true">
+          <ID>paramName2</ID>
+          <ToolTip>Parameter name</ToolTip>
+          <Default>Param2Name</Default>
+          <Function>
+          </Function>
+            </Literal>
+        <Literal Editable="true">
+          <ID>paramName3</ID>
+          <ToolTip>Parameter name</ToolTip>
+          <Default>Param3Name</Default>
+          <Function>
+          </Function>
+            </Literal>
+       <Literal Editable="true">
+          <ID>paramType1</ID>
+          <ToolTip>Parameter type</ToolTip>
+          <Default>Param1Type</Default>
+          <Function>
+          </Function>
+            </Literal>
+       <Literal Editable="true">
+          <ID>paramType2</ID>
+          <ToolTip>Parameter type</ToolTip>
+          <Default>Param2Type</Default>
+          <Function>
+          </Function>
+            </Literal>
+       <Literal Editable="true">
+          <ID>paramType3</ID>
+          <ToolTip>Parameter type</ToolTip>
+          <Default>Param3Type</Default>
+          <Function>
+          </Function>
+            </Literal>
+        </Declarations>
+          <Code Language="cpp"><![CDATA[DECLARE_DYNAMIC_MULTICAST_DELEGATE_ThreeParams($delegateName$, $paramType1$, $paramName1$, $paramType2$, $paramName2$, $paramType3$, $paramName3$)$selected$ $end$]]></Code>
+        </Snippet>
+      </CodeSnippet>
+    </CodeSnippets>
\ No newline at end of file
diff --git a/Engine/Extras/VisualStudioSnippets/ue4deldynmul4.snippet b/Engine/Extras/VisualStudioSnippets/ue4deldynmul4.snippet
new file mode 100644
--- /dev/null
+++ b/Engine/Extras/VisualStudioSnippets/ue4deldynmul4.snippet
@@ -0,0 +1,84 @@
+<?xml version="1.0" encoding="utf-8"?>
+    <CodeSnippets xmlns="http://schemas.microsoft.com/VisualStudio/2005/CodeSnippet">
+      <CodeSnippet Format="1.0.0">
+        <Header>
+          <SnippetTypes>
+            <SnippetType>Expansion</SnippetType>
+          </SnippetTypes>
+          <Title>ue4deldynmul4</Title>
+          <Author>Epic Games</Author>
+          <Description>Unreal Engine Dynamic Multicast Delegate snippet</Description>
+          <HelpUrl>
+          </HelpUrl>
+          <Shortcut>ue4deldynmul4</Shortcut>
+        </Header>
+        <Snippet>
+          <Declarations>
+       <Literal Editable="true">
+          <ID>delegateName</ID>
+          <ToolTip>Delegate type name</ToolTip>
+          <Default>DelegateName</Default>
+          <Function>
+          </Function>
+        </Literal>
+        <Literal Editable="true">
+          <ID>paramName1</ID>
+          <ToolTip>Parameter name</ToolTip>
+          <Default>Param1Name</Default>
+          <Function>
+          </Function>
+            </Literal>
+        <Literal Editable="true">
+          <ID>paramName2</ID>
+          <ToolTip>Parameter name</ToolTip>
+          <Default>Param2Name</Default>
+          <Function>
+          </Function>
+            </Literal>
+        <Literal Editable="true">
+          <ID>paramName3</ID>
+          <ToolTip>Parameter name</ToolTip>
+          <Default>Param3Name</Default>
+          <Function>
+          </Function>
+            </Literal>
+        <Literal Editable="true">
+          <ID>paramName4</ID>
+          <ToolTip>Parameter name</ToolTip>
+          <Default>Param4Name</Default>
+          <Function>
+          </Function>
+            </Literal>
+       <Literal Editable="true">
+          <ID>paramType1</ID>
+          <ToolTip>Parameter type</ToolTip>
+          <Default>Param1Type</Default>
+          <Function>
+          </Function>
+            </Literal>
+       <Literal Editable="true">
+          <ID>paramType2</ID>
+          <ToolTip>Parameter type</ToolTip>
+          <Default>Param2Type</Default>
+          <Function>
+          </Function>
+            </Literal>
+       <Literal Editable="true">
+          <ID>paramType3</ID>
+          <ToolTip>Parameter type</ToolTip>
+          <Default>Param3Type</Default>
+          <Function>
+          </Function>
+            </Literal>
+       <Literal Editable="true">
+          <ID>paramType4</ID>
+          <ToolTip>Parameter type</ToolTip>
+          <Default>Param4Type</Default>
+          <Function>
+          </Function>
+            </Literal>
+        </Declarations>
+          <Code Language="cpp"><![CDATA[DECLARE_DYNAMIC_MULTICAST_DELEGATE_FourParams($delegateName$, $paramType1$, $paramName1$, $paramType2$, $paramName2$, $paramType3$, $paramName3$, $paramType4$, $paramName4$)$selected$ $end$]]></Code>
+        </Snippet>
+      </CodeSnippet>
+    </CodeSnippets>
\ No newline at end of file
diff --git a/Engine/Extras/VisualStudioSnippets/ue4deldynmul5.snippet b/Engine/Extras/VisualStudioSnippets/ue4deldynmul5.snippet
new file mode 100644
--- /dev/null
+++ b/Engine/Extras/VisualStudioSnippets/ue4deldynmul5.snippet
@@ -0,0 +1,98 @@
+<?xml version="1.0" encoding="utf-8"?>
+    <CodeSnippets xmlns="http://schemas.microsoft.com/VisualStudio/2005/CodeSnippet">
+      <CodeSnippet Format="1.0.0">
+        <Header>
+          <SnippetTypes>
+            <SnippetType>Expansion</SnippetType>
+          </SnippetTypes>
+          <Title>ue4deldynmul5</Title>
+          <Author>Epic Games</Author>
+          <Description>Unreal Engine Dynamic Multicast Delegate snippet</Description>
+          <HelpUrl>
+          </HelpUrl>
+          <Shortcut>ue4deldynmul5</Shortcut>
+        </Header>
+        <Snippet>
+          <Declarations>
+       <Literal Editable="true">
+          <ID>delegateName</ID>
+          <ToolTip>Delegate type name</ToolTip>
+          <Default>DelegateName</Default>
+          <Function>
+          </Function>
+        </Literal>
+        <Literal Editable="true">
+          <ID>paramName1</ID>
+          <ToolTip>Parameter name</ToolTip>
+          <Default>Param1Name</Default>
+          <Function>
+          </Function>
+            </Literal>
+        <Literal Editable="true">
+          <ID>paramName2</ID>
+          <ToolTip>Parameter name</ToolTip>
+          <Default>Param2Name</Default>
+          <Function>
+          </Function>
+            </Literal>
+        <Literal Editable="true">
+          <ID>paramName3</ID>
+          <ToolTip>Parameter name</ToolTip>
+          <Default>Param3Name</Default>
+          <Function>
+          </Function>
+            </Literal>
+        <Literal Editable="true">
+          <ID>paramName4</ID>
+          <ToolTip>Parameter name</ToolTip>
+          <Default>Param4Name</Default>
+          <Function>
+          </Function>
+            </Literal>
+        <Literal Editable="true">
+          <ID>paramName5</ID>
+          <ToolTip>Parameter name</ToolTip>
+          <Default>Param5Name</Default>
+          <Function>
+          </Function>
+            </Literal>
+       <Literal Editable="true">
+          <ID>paramType1</ID>
+          <ToolTip>Parameter type</ToolTip>
+          <Default>Param1Type</Default>
+          <Function>
+          </Function>
+            </Literal>
+       <Literal Editable="true">
+          <ID>paramType2</ID>
+          <ToolTip>Parameter type</ToolTip>
+          <Default>Param2Type</Default>
+          <Function>
+          </Function>
+            </Literal>
+       <Literal Editable="true">
+          <ID>paramType3</ID>
+          <ToolTip>Parameter type</ToolTip>
+          <Default>Param3Type</Default>
+          <Function>
+          </Function>
+            </Literal>
+       <Literal Editable="true">
+          <ID>paramType4</ID>
+          <ToolTip>Parameter type</ToolTip>
+          <Default>Param4Type</Default>
+          <Function>
+          </Function>
+            </Literal>
+       <Literal Editable="true">
+          <ID>paramType5</ID>
+          <ToolTip>Parameter type</ToolTip>
+          <Default>Param5Type</Default>
+          <Function>
+          </Function>
+            </Literal>
+        </Declarations>
+          <Code Language="cpp"><![CDATA[DECLARE_DYNAMIC_MULTICAST_DELEGATE_FiveParams($delegateName$, $paramType1$, $paramName1$, $paramType2$, $paramName2$, $paramType3$, $paramName3$, $paramType4$, $paramName4$, $paramType5$, $paramName5$)$selected$ $end$]]></Code>
+        </Snippet>
+      </CodeSnippet>
+    </CodeSnippets>
\ No newline at end of file
diff --git a/Engine/Extras/VisualStudioSnippets/ue4deldynmul6.snippet b/Engine/Extras/VisualStudioSnippets/ue4deldynmul6.snippet
new file mode 100644
--- /dev/null
+++ b/Engine/Extras/VisualStudioSnippets/ue4deldynmul6.snippet
@@ -0,0 +1,112 @@
+<?xml version="1.0" encoding="utf-8"?>
+    <CodeSnippets xmlns="http://schemas.microsoft.com/VisualStudio/2005/CodeSnippet">
+      <CodeSnippet Format="1.0.0">
+        <Header>
+          <SnippetTypes>
+            <SnippetType>Expansion</SnippetType>
+          </SnippetTypes>
+          <Title>ue4deldynmul6</Title>
+          <Author>Epic Games</Author>
+          <Description>Unreal Engine Dynamic Multicast Delegate snippet</Description>
+          <HelpUrl>
+          </HelpUrl>
+          <Shortcut>ue4deldynmul6</Shortcut>
+        </Header>
+        <Snippet>
+          <Declarations>
+       <Literal Editable="true">
+          <ID>delegateName</ID>
+          <ToolTip>Delegate type name</ToolTip>
+          <Default>DelegateName</Default>
+          <Function>
+          </Function>
+        </Literal>
+        <Literal Editable="true">
+          <ID>paramName1</ID>
+          <ToolTip>Parameter name</ToolTip>
+          <Default>Param1Name</Default>
+          <Function>
+          </Function>
+            </Literal>
+        <Literal Editable="true">
+          <ID>paramName2</ID>
+          <ToolTip>Parameter name</ToolTip>
+          <Default>Param2Name</Default>
+          <Function>
+          </Function>
+            </Literal>
+        <Literal Editable="true">
+          <ID>paramName3</ID>
+          <ToolTip>Parameter name</ToolTip>
+          <Default>Param3Name</Default>
+          <Function>
+          </Function>
+            </Literal>
+        <Literal Editable="true">
+          <ID>paramName4</ID>
+          <ToolTip>Parameter name</ToolTip>
+          <Default>Param4Name</Default>
+          <Function>
+          </Function>
+            </Literal>
+        <Literal Editable="true">
+          <ID>paramName5</ID>
+          <ToolTip>Parameter name</ToolTip>
+          <Default>Param5Name</Default>
+          <Function>
+          </Function>
+            </Literal>
+        <Literal Editable="true">
+          <ID>paramName6</ID>
+          <ToolTip>Parameter name</ToolTip>
+          <Default>Param6Name</Default>
+          <Function>
+          </Function>
+            </Literal>
+       <Literal Editable="true">
+          <ID>paramType1</ID>
+          <ToolTip>Parameter type</ToolTip>
+          <Default>Param1Type</Default>
+          <Function>
+          </Function>
+            </Literal>
+       <Literal Editable="true">
+          <ID>paramType2</ID>
+          <ToolTip>Parameter type</ToolTip>
+          <Default>Param2Type</Default>
+          <Function>
+          </Function>
+            </Literal>
+       <Literal Editable="true">
+          <ID>paramType3</ID>
+          <ToolTip>Parameter type</ToolTip>
+          <Default>Param3Type</Default>
+          <Function>
+          </Function>
+            </Literal>
+       <Literal Editable="true">
+          <ID>paramType4</ID>
+          <ToolTip>Parameter type</ToolTip>
+          <Default>Param4Type</Default>
+          <Function>
+          </Function>
+            </Literal>
+       <Literal Editable="true">
+          <ID>paramType5</ID>
+          <ToolTip>Parameter type</ToolTip>
+          <Default>Param5Type</Default>
+          <Function>
+          </Function>
+            </Literal>
+       <Literal Editable="true">
+          <ID>paramType6</ID>
+          <ToolTip>Parameter type</ToolTip>
+          <Default>Param6Type</Default>
+          <Function>
+          </Function>
+            </Literal>
+        </Declarations>
+          <Code Language="cpp"><![CDATA[DECLARE_DYNAMIC_MULTICAST_DELEGATE_SixParams($delegateName$, $paramType1$, $paramName1$, $paramType2$, $paramName2$, $paramType3$, $paramName3$, $paramType4$, $paramName4$, $paramType5$, $paramName5$, $paramType6$, $paramName6$)$selected$ $end$]]></Code>
+        </Snippet>
+      </CodeSnippet>
+    </CodeSnippets>
\ No newline at end of file
diff --git a/Engine/Extras/VisualStudioSnippets/ue4deldynmul7.snippet b/Engine/Extras/VisualStudioSnippets/ue4deldynmul7.snippet
new file mode 100644
--- /dev/null
+++ b/Engine/Extras/VisualStudioSnippets/ue4deldynmul7.snippet
@@ -0,0 +1,126 @@
+<?xml version="1.0" encoding="utf-8"?>
+    <CodeSnippets xmlns="http://schemas.microsoft.com/VisualStudio/2005/CodeSnippet">
+      <CodeSnippet Format="1.0.0">
+        <Header>
+          <SnippetTypes>
+            <SnippetType>Expansion</SnippetType>
+          </SnippetTypes>
+          <Title>ue4deldynmul7</Title>
+          <Author>Epic Games</Author>
+          <Description>Unreal Engine Dynamic Multicast Delegate snippet</Description>
+          <HelpUrl>
+          </HelpUrl>
+          <Shortcut>ue4deldynmul7</Shortcut>
+        </Header>
+        <Snippet>
+          <Declarations>
+       <Literal Editable="true">
+          <ID>delegateName</ID>
+          <ToolTip>Delegate type name</ToolTip>
+          <Default>DelegateName</Default>
+          <Function>
+          </Function>
+        </Literal>
+        <Literal Editable="true">
+          <ID>paramName1</ID>
+          <ToolTip>Parameter name</ToolTip>
+          <Default>Param1Name</Default>
+          <Function>
+          </Function>
+            </Literal>
+        <Literal Editable="true">
+          <ID>paramName2</ID>
+          <ToolTip>Parameter name</ToolTip>
+          <Default>Param2Name</Default>
+          <Function>
+          </Function>
+            </Literal>
+        <Literal Editable="true">
+          <ID>paramName3</ID>
+          <ToolTip>Parameter name</ToolTip>
+          <Default>Param3Name</Default>
+          <Function>
+          </Function>
+            </Literal>
+        <Literal Editable="true">
+          <ID>paramName4</ID>
+          <ToolTip>Parameter name</ToolTip>
+          <Default>Param4Name</Default>
+          <Function>
+          </Function>
+            </Literal>
+        <Literal Editable="true">
+          <ID>paramName5</ID>
+          <ToolTip>Parameter name</ToolTip>
+          <Default>Param5Name</Default>
+          <Function>
+          </Function>
+            </Literal>
+        <Literal Editable="true">
+          <ID>paramName6</ID>
+          <ToolTip>Parameter name</ToolTip>
+          <Default>Param6Name</Default>
+          <Function>
+          </Function>
+            </Literal>
+        <Literal Editable="true">
+          <ID>paramName7</ID>
+          <ToolTip>Parameter name</ToolTip>
+          <Default>Param7Name</Default>
+          <Function>
+          </Function>
+            </Literal>
+       <Literal Editable="true">
+          <ID>paramType1</ID>
+          <ToolTip>Parameter type</ToolTip>
+          <Default>Param1Type</Default>
+          <Function>
+          </Function>
+            </Literal>
+       <Literal Editable="true">
+          <ID>paramType2</ID>
+          <ToolTip>Parameter type</ToolTip>
+          <Default>Param2Type</Default>
+          <Function>
+          </Function>
+            </Literal>
+       <Literal Editable="true">
+          <ID>paramType3</ID>
+          <ToolTip>Parameter type</ToolTip>
+          <Default>Param3Type</Default>
+          <Function>
+          </Function>
+            </Literal>
+       <Literal Editable="true">
+          <ID>paramType4</ID>
+          <ToolTip>Parameter type</ToolTip>
+          <Default>Param4Type</Default>
+          <Function>
+          </Function>
+            </Literal>
+       <Literal Editable="true">
+          <ID>paramType5</ID>
+          <ToolTip>Parameter type</ToolTip>
+          <Default>Param5Type</Default>
+          <Function>
+          </Function>
+            </Literal>
+       <Literal Editable="true">
+          <ID>paramType6</ID>
+          <ToolTip>Parameter type</ToolTip>
+          <Default>Param6Type</Default>
+          <Function>
+          </Function>
+            </Literal>
+       <Literal Editable="true">
+          <ID>paramType7</ID>
+          <ToolTip>Parameter type</ToolTip>
+          <Default>Param7Type</Default>
+          <Function>
+          </Function>
+            </Literal>
+        </Declarations>
+          <Code Language="cpp"><![CDATA[DECLARE_DYNAMIC_MULTICAST_DELEGATE_SevenParams($delegateName$, $paramType1$, $paramName1$, $paramType2$, $paramName2$, $paramType3$, $paramName3$, $paramType4$, $paramName4$, $paramType5$, $paramName5$, $paramType6$, $paramName6$, $paramType7$, $paramName7$)$selected$ $end$]]></Code>
+        </Snippet>
+      </CodeSnippet>
+    </CodeSnippets>
\ No newline at end of file
diff --git a/Engine/Extras/VisualStudioSnippets/ue4deldynmul8.snippet b/Engine/Extras/VisualStudioSnippets/ue4deldynmul8.snippet
new file mode 100644
--- /dev/null
+++ b/Engine/Extras/VisualStudioSnippets/ue4deldynmul8.snippet
@@ -0,0 +1,140 @@
+<?xml version="1.0" encoding="utf-8"?>
+    <CodeSnippets xmlns="http://schemas.microsoft.com/VisualStudio/2005/CodeSnippet">
+      <CodeSnippet Format="1.0.0">
+        <Header>
+          <SnippetTypes>
+            <SnippetType>Expansion</SnippetType>
+          </SnippetTypes>
+          <Title>ue4deldynmul8</Title>
+          <Author>Epic Games</Author>
+          <Description>Unreal Engine Dynamic Multicast Delegate snippet</Description>
+          <HelpUrl>
+          </HelpUrl>
+          <Shortcut>ue4deldynmul8</Shortcut>
+        </Header>
+        <Snippet>
+          <Declarations>
+       <Literal Editable="true">
+          <ID>delegateName</ID>
+          <ToolTip>Delegate type name</ToolTip>
+          <Default>DelegateName</Default>
+          <Function>
+          </Function>
+        </Literal>
+        <Literal Editable="true">
+          <ID>paramName1</ID>
+          <ToolTip>Parameter name</ToolTip>
+          <Default>Param1Name</Default>
+          <Function>
+          </Function>
+            </Literal>
+        <Literal Editable="true">
+          <ID>paramName2</ID>
+          <ToolTip>Parameter name</ToolTip>
+          <Default>Param2Name</Default>
+          <Function>
+          </Function>
+            </Literal>
+        <Literal Editable="true">
+          <ID>paramName3</ID>
+          <ToolTip>Parameter name</ToolTip>
+          <Default>Param3Name</Default>
+          <Function>
+          </Function>
+            </Literal>
+        <Literal Editable="true">
+          <ID>paramName4</ID>
+          <ToolTip>Parameter name</ToolTip>
+          <Default>Param4Name</Default>
+          <Function>
+          </Function>
+            </Literal>
+        <Literal Editable="true">
+          <ID>paramName5</ID>
+          <ToolTip>Parameter name</ToolTip>
+          <Default>Param5Name</Default>
+          <Function>
+          </Function>
+            </Literal>
+        <Literal Editable="true">
+          <ID>paramName6</ID>
+          <ToolTip>Parameter name</ToolTip>
+          <Default>Param6Name</Default>
+          <Function>
+          </Function>
+            </Literal>
+        <Literal Editable="true">
+          <ID>paramName7</ID>
+          <ToolTip>Parameter name</ToolTip>
+          <Default>Param7Name</Default>
+          <Function>
+          </Function>
+            </Literal>
+        <Literal Editable="true">
+          <ID>paramName8</ID>
+          <ToolTip>Parameter name</ToolTip>
+          <Default>Param8Name</Default>
+          <Function>
+          </Function>
+            </Literal>
+       <Literal Editable="true">
+          <ID>paramType1</ID>
+          <ToolTip>Parameter type</ToolTip>
+          <Default>Param1Type</Default>
+          <Function>
+          </Function>
+            </Literal>
+       <Literal Editable="true">
+          <ID>paramType2</ID>
+          <ToolTip>Parameter type</ToolTip>
+          <Default>Param2Type</Default>
+          <Function>
+          </Function>
+            </Literal>
+       <Literal Editable="true">
+          <ID>paramType3</ID>
+          <ToolTip>Parameter type</ToolTip>
+          <Default>Param3Type</Default>
+          <Function>
+          </Function>
+            </Literal>
+       <Literal Editable="true">
+          <ID>paramType4</ID>
+          <ToolTip>Parameter type</ToolTip>
+          <Default>Param4Type</Default>
+          <Function>
+          </Function>
+            </Literal>
+       <Literal Editable="true">
+          <ID>paramType5</ID>
+          <ToolTip>Parameter type</ToolTip>
+          <Default>Param5Type</Default>
+          <Function>
+          </Function>
+            </Literal>
+       <Literal Editable="true">
+          <ID>paramType6</ID>
+          <ToolTip>Parameter type</ToolTip>
+          <Default>Param6Type</Default>
+          <Function>
+          </Function>
+            </Literal>
+       <Literal Editable="true">
+          <ID>paramType7</ID>
+          <ToolTip>Parameter type</ToolTip>
+          <Default>Param7Type</Default>
+          <Function>
+          </Function>
+            </Literal>
+       <Literal Editable="true">
+          <ID>paramType8</ID>
+          <ToolTip>Parameter type</ToolTip>
+          <Default>Param8Type</Default>
+          <Function>
+          </Function>
+            </Literal>
+        </Declarations>
+          <Code Language="cpp"><![CDATA[DECLARE_DYNAMIC_MULTICAST_DELEGATE_EightParams($delegateName$, $paramType1$, $paramName1$, $paramType2$, $paramName2$, $paramType3$, $paramName3$, $paramType4$, $paramName4$, $paramType5$, $paramName5$, $paramType6$, $paramName6$, $paramType7$, $paramName7$, $paramType8$, $paramName8$)$selected$ $end$]]></Code>
+        </Snippet>
+      </CodeSnippet>
+    </CodeSnippets>
\ No newline at end of file
diff --git a/Engine/Extras/VisualStudioSnippets/ue4delmul.snippet b/Engine/Extras/VisualStudioSnippets/ue4delmul.snippet
new file mode 100644
--- /dev/null
+++ b/Engine/Extras/VisualStudioSnippets/ue4delmul.snippet
@@ -0,0 +1,28 @@
+<?xml version="1.0" encoding="utf-8"?>
+    <CodeSnippets xmlns="http://schemas.microsoft.com/VisualStudio/2005/CodeSnippet">
+      <CodeSnippet Format="1.0.0">
+        <Header>
+          <SnippetTypes>
+            <SnippetType>Expansion</SnippetType>
+          </SnippetTypes>
+          <Title>ue4delmul</Title>
+          <Author>Epic Games</Author>
+          <Description>Unreal Engine Multicast Delegate snippet</Description>
+          <HelpUrl>
+          </HelpUrl>
+          <Shortcut>ue4delmul</Shortcut>
+        </Header>
+        <Snippet>
+          <Declarations>
+       <Literal Editable="true">
+          <ID>delegateName</ID>
+          <ToolTip>Delegate type name</ToolTip>
+          <Default>DelegateName</Default>
+          <Function>
+          </Function>
+        </Literal>
+        </Declarations>
+          <Code Language="cpp"><![CDATA[DECLARE_MULTICAST_DELEGATE($delegateName$)$selected$ $end$]]></Code>
+        </Snippet>
+      </CodeSnippet>
+    </CodeSnippets>
\ No newline at end of file
diff --git a/Engine/Extras/VisualStudioSnippets/ue4delmul1.snippet b/Engine/Extras/VisualStudioSnippets/ue4delmul1.snippet
new file mode 100644
--- /dev/null
+++ b/Engine/Extras/VisualStudioSnippets/ue4delmul1.snippet
@@ -0,0 +1,42 @@
+<?xml version="1.0" encoding="utf-8"?>
+    <CodeSnippets xmlns="http://schemas.microsoft.com/VisualStudio/2005/CodeSnippet">
+      <CodeSnippet Format="1.0.0">
+        <Header>
+          <SnippetTypes>
+            <SnippetType>Expansion</SnippetType>
+          </SnippetTypes>
+          <Title>ue4delmul1</Title>
+          <Author>Epic Games</Author>
+          <Description>Unreal Engine Multicast Delegate snippet</Description>
+          <HelpUrl>
+          </HelpUrl>
+          <Shortcut>ue4delmul1</Shortcut>
+        </Header>
+        <Snippet>
+          <Declarations>
+       <Literal Editable="true">
+          <ID>delegateName</ID>
+          <ToolTip>Delegate type name</ToolTip>
+          <Default>DelegateName</Default>
+          <Function>
+          </Function>
+        </Literal>
+        <Literal Editable="true">
+          <ID>paramName1</ID>
+          <ToolTip>Parameter name</ToolTip>
+          <Default>Param1Name</Default>
+          <Function>
+          </Function>
+            </Literal>
+       <Literal Editable="true">
+          <ID>paramType1</ID>
+          <ToolTip>Parameter type</ToolTip>
+          <Default>Param1Type</Default>
+          <Function>
+          </Function>
+            </Literal>
+        </Declarations>
+          <Code Language="cpp"><![CDATA[DECLARE_MULTICAST_DELEGATE_OneParam($delegateName$, $paramType1$, $paramName1$)$selected$ $end$]]></Code>
+        </Snippet>
+      </CodeSnippet>
+    </CodeSnippets>
\ No newline at end of file
diff --git a/Engine/Extras/VisualStudioSnippets/ue4delmul2.snippet b/Engine/Extras/VisualStudioSnippets/ue4delmul2.snippet
new file mode 100644
--- /dev/null
+++ b/Engine/Extras/VisualStudioSnippets/ue4delmul2.snippet
@@ -0,0 +1,56 @@
+<?xml version="1.0" encoding="utf-8"?>
+    <CodeSnippets xmlns="http://schemas.microsoft.com/VisualStudio/2005/CodeSnippet">
+      <CodeSnippet Format="1.0.0">
+        <Header>
+          <SnippetTypes>
+            <SnippetType>Expansion</SnippetType>
+          </SnippetTypes>
+          <Title>ue4delmul2</Title>
+          <Author>Epic Games</Author>
+          <Description>Unreal Engine Multicast Delegate snippet</Description>
+          <HelpUrl>
+          </HelpUrl>
+          <Shortcut>ue4delmul2</Shortcut>
+        </Header>
+        <Snippet>
+          <Declarations>
+       <Literal Editable="true">
+          <ID>delegateName</ID>
+          <ToolTip>Delegate type name</ToolTip>
+          <Default>DelegateName</Default>
+          <Function>
+          </Function>
+        </Literal>
+        <Literal Editable="true">
+          <ID>paramName1</ID>
+          <ToolTip>Parameter name</ToolTip>
+          <Default>Param1Name</Default>
+          <Function>
+          </Function>
+            </Literal>
+        <Literal Editable="true">
+          <ID>paramName2</ID>
+          <ToolTip>Parameter name</ToolTip>
+          <Default>Param2Name</Default>
+          <Function>
+          </Function>
+            </Literal>
+       <Literal Editable="true">
+          <ID>paramType1</ID>
+          <ToolTip>Parameter type</ToolTip>
+          <Default>Param1Type</Default>
+          <Function>
+          </Function>
+            </Literal>
+       <Literal Editable="true">
+          <ID>paramType2</ID>
+          <ToolTip>Parameter type</ToolTip>
+          <Default>Param2Type</Default>
+          <Function>
+          </Function>
+            </Literal>
+        </Declarations>
+          <Code Language="cpp"><![CDATA[DECLARE_MULTICAST_DELEGATE_TwoParams($delegateName$, $paramType1$, $paramName1$, $paramType2$, $paramName2$)$selected$ $end$]]></Code>
+        </Snippet>
+      </CodeSnippet>
+    </CodeSnippets>
\ No newline at end of file
diff --git a/Engine/Extras/VisualStudioSnippets/ue4delmul3.snippet b/Engine/Extras/VisualStudioSnippets/ue4delmul3.snippet
new file mode 100644
--- /dev/null
+++ b/Engine/Extras/VisualStudioSnippets/ue4delmul3.snippet
@@ -0,0 +1,70 @@
+<?xml version="1.0" encoding="utf-8"?>
+    <CodeSnippets xmlns="http://schemas.microsoft.com/VisualStudio/2005/CodeSnippet">
+      <CodeSnippet Format="1.0.0">
+        <Header>
+          <SnippetTypes>
+            <SnippetType>Expansion</SnippetType>
+          </SnippetTypes>
+          <Title>ue4delmul3</Title>
+          <Author>Epic Games</Author>
+          <Description>Unreal Engine Multicast Delegate snippet</Description>
+          <HelpUrl>
+          </HelpUrl>
+          <Shortcut>ue4delmul3</Shortcut>
+        </Header>
+        <Snippet>
+          <Declarations>
+       <Literal Editable="true">
+          <ID>delegateName</ID>
+          <ToolTip>Delegate type name</ToolTip>
+          <Default>DelegateName</Default>
+          <Function>
+          </Function>
+        </Literal>
+        <Literal Editable="true">
+          <ID>paramName1</ID>
+          <ToolTip>Parameter name</ToolTip>
+          <Default>Param1Name</Default>
+          <Function>
+          </Function>
+            </Literal>
+        <Literal Editable="true">
+          <ID>paramName2</ID>
+          <ToolTip>Parameter name</ToolTip>
+          <Default>Param2Name</Default>
+          <Function>
+          </Function>
+            </Literal>
+        <Literal Editable="true">
+          <ID>paramName3</ID>
+          <ToolTip>Parameter name</ToolTip>
+          <Default>Param3Name</Default>
+          <Function>
+          </Function>
+            </Literal>
+       <Literal Editable="true">
+          <ID>paramType1</ID>
+          <ToolTip>Parameter type</ToolTip>
+          <Default>Param1Type</Default>
+          <Function>
+          </Function>
+            </Literal>
+       <Literal Editable="true">
+          <ID>paramType2</ID>
+          <ToolTip>Parameter type</ToolTip>
+          <Default>Param2Type</Default>
+          <Function>
+          </Function>
+            </Literal>
+       <Literal Editable="true">
+          <ID>paramType3</ID>
+          <ToolTip>Parameter type</ToolTip>
+          <Default>Param3Type</Default>
+          <Function>
+          </Function>
+            </Literal>
+        </Declarations>
+          <Code Language="cpp"><![CDATA[DECLARE_MULTICAST_DELEGATE_ThreeParams($delegateName$, $paramType1$, $paramName1$, $paramType2$, $paramName2$, $paramType3$, $paramName3$)$selected$ $end$]]></Code>
+        </Snippet>
+      </CodeSnippet>
+    </CodeSnippets>
\ No newline at end of file
diff --git a/Engine/Extras/VisualStudioSnippets/ue4delmul4.snippet b/Engine/Extras/VisualStudioSnippets/ue4delmul4.snippet
new file mode 100644
--- /dev/null
+++ b/Engine/Extras/VisualStudioSnippets/ue4delmul4.snippet
@@ -0,0 +1,84 @@
+<?xml version="1.0" encoding="utf-8"?>
+    <CodeSnippets xmlns="http://schemas.microsoft.com/VisualStudio/2005/CodeSnippet">
+      <CodeSnippet Format="1.0.0">
+        <Header>
+          <SnippetTypes>
+            <SnippetType>Expansion</SnippetType>
+          </SnippetTypes>
+          <Title>ue4delmul4</Title>
+          <Author>Epic Games</Author>
+          <Description>Unreal Engine Multicast Delegate snippet</Description>
+          <HelpUrl>
+          </HelpUrl>
+          <Shortcut>ue4delmul4</Shortcut>
+        </Header>
+        <Snippet>
+          <Declarations>
+       <Literal Editable="true">
+          <ID>delegateName</ID>
+          <ToolTip>Delegate type name</ToolTip>
+          <Default>DelegateName</Default>
+          <Function>
+          </Function>
+        </Literal>
+        <Literal Editable="true">
+          <ID>paramName1</ID>
+          <ToolTip>Parameter name</ToolTip>
+          <Default>Param1Name</Default>
+          <Function>
+          </Function>
+            </Literal>
+        <Literal Editable="true">
+          <ID>paramName2</ID>
+          <ToolTip>Parameter name</ToolTip>
+          <Default>Param2Name</Default>
+          <Function>
+          </Function>
+            </Literal>
+        <Literal Editable="true">
+          <ID>paramName3</ID>
+          <ToolTip>Parameter name</ToolTip>
+          <Default>Param3Name</Default>
+          <Function>
+          </Function>
+            </Literal>
+        <Literal Editable="true">
+          <ID>paramName4</ID>
+          <ToolTip>Parameter name</ToolTip>
+          <Default>Param4Name</Default>
+          <Function>
+          </Function>
+            </Literal>
+       <Literal Editable="true">
+          <ID>paramType1</ID>
+          <ToolTip>Parameter type</ToolTip>
+          <Default>Param1Type</Default>
+          <Function>
+          </Function>
+            </Literal>
+       <Literal Editable="true">
+          <ID>paramType2</ID>
+          <ToolTip>Parameter type</ToolTip>
+          <Default>Param2Type</Default>
+          <Function>
+          </Function>
+            </Literal>
+       <Literal Editable="true">
+          <ID>paramType3</ID>
+          <ToolTip>Parameter type</ToolTip>
+          <Default>Param3Type</Default>
+          <Function>
+          </Function>
+            </Literal>
+       <Literal Editable="true">
+          <ID>paramType4</ID>
+          <ToolTip>Parameter type</ToolTip>
+          <Default>Param4Type</Default>
+          <Function>
+          </Function>
+            </Literal>
+        </Declarations>
+          <Code Language="cpp"><![CDATA[DECLARE_MULTICAST_DELEGATE_FourParams($delegateName$, $paramType1$, $paramName1$, $paramType2$, $paramName2$, $paramType3$, $paramName3$, $paramType4$, $paramName4$)$selected$ $end$]]></Code>
+        </Snippet>
+      </CodeSnippet>
+    </CodeSnippets>
\ No newline at end of file
diff --git a/Engine/Extras/VisualStudioSnippets/ue4delmul5.snippet b/Engine/Extras/VisualStudioSnippets/ue4delmul5.snippet
new file mode 100644
--- /dev/null
+++ b/Engine/Extras/VisualStudioSnippets/ue4delmul5.snippet
@@ -0,0 +1,98 @@
+<?xml version="1.0" encoding="utf-8"?>
+    <CodeSnippets xmlns="http://schemas.microsoft.com/VisualStudio/2005/CodeSnippet">
+      <CodeSnippet Format="1.0.0">
+        <Header>
+          <SnippetTypes>
+            <SnippetType>Expansion</SnippetType>
+          </SnippetTypes>
+          <Title>ue4delmul5</Title>
+          <Author>Epic Games</Author>
+          <Description>Unreal Engine Multicast Delegate snippet</Description>
+          <HelpUrl>
+          </HelpUrl>
+          <Shortcut>ue4delmul5</Shortcut>
+        </Header>
+        <Snippet>
+          <Declarations>
+       <Literal Editable="true">
+          <ID>delegateName</ID>
+          <ToolTip>Delegate type name</ToolTip>
+          <Default>DelegateName</Default>
+          <Function>
+          </Function>
+        </Literal>
+        <Literal Editable="true">
+          <ID>paramName1</ID>
+          <ToolTip>Parameter name</ToolTip>
+          <Default>Param1Name</Default>
+          <Function>
+          </Function>
+            </Literal>
+        <Literal Editable="true">
+          <ID>paramName2</ID>
+          <ToolTip>Parameter name</ToolTip>
+          <Default>Param2Name</Default>
+          <Function>
+          </Function>
+            </Literal>
+        <Literal Editable="true">
+          <ID>paramName3</ID>
+          <ToolTip>Parameter name</ToolTip>
+          <Default>Param3Name</Default>
+          <Function>
+          </Function>
+            </Literal>
+        <Literal Editable="true">
+          <ID>paramName4</ID>
+          <ToolTip>Parameter name</ToolTip>
+          <Default>Param4Name</Default>
+          <Function>
+          </Function>
+            </Literal>
+        <Literal Editable="true">
+          <ID>paramName5</ID>
+          <ToolTip>Parameter name</ToolTip>
+          <Default>Param5Name</Default>
+          <Function>
+          </Function>
+            </Literal>
+       <Literal Editable="true">
+          <ID>paramType1</ID>
+          <ToolTip>Parameter type</ToolTip>
+          <Default>Param1Type</Default>
+          <Function>
+          </Function>
+            </Literal>
+       <Literal Editable="true">
+          <ID>paramType2</ID>
+          <ToolTip>Parameter type</ToolTip>
+          <Default>Param2Type</Default>
+          <Function>
+          </Function>
+            </Literal>
+       <Literal Editable="true">
+          <ID>paramType3</ID>
+          <ToolTip>Parameter type</ToolTip>
+          <Default>Param3Type</Default>
+          <Function>
+          </Function>
+            </Literal>
+       <Literal Editable="true">
+          <ID>paramType4</ID>
+          <ToolTip>Parameter type</ToolTip>
+          <Default>Param4Type</Default>
+          <Function>
+          </Function>
+            </Literal>
+       <Literal Editable="true">
+          <ID>paramType5</ID>
+          <ToolTip>Parameter type</ToolTip>
+          <Default>Param5Type</Default>
+          <Function>
+          </Function>
+            </Literal>
+        </Declarations>
+          <Code Language="cpp"><![CDATA[DECLARE_MULTICAST_DELEGATE_FiveParams($delegateName$, $paramType1$, $paramName1$, $paramType2$, $paramName2$, $paramType3$, $paramName3$, $paramType4$, $paramName4$, $paramType5$, $paramName5$)$selected$ $end$]]></Code>
+        </Snippet>
+      </CodeSnippet>
+    </CodeSnippets>
\ No newline at end of file
diff --git a/Engine/Extras/VisualStudioSnippets/ue4delmul6.snippet b/Engine/Extras/VisualStudioSnippets/ue4delmul6.snippet
new file mode 100644
--- /dev/null
+++ b/Engine/Extras/VisualStudioSnippets/ue4delmul6.snippet
@@ -0,0 +1,112 @@
+<?xml version="1.0" encoding="utf-8"?>
+    <CodeSnippets xmlns="http://schemas.microsoft.com/VisualStudio/2005/CodeSnippet">
+      <CodeSnippet Format="1.0.0">
+        <Header>
+          <SnippetTypes>
+            <SnippetType>Expansion</SnippetType>
+          </SnippetTypes>
+          <Title>ue4delmul6</Title>
+          <Author>Epic Games</Author>
+          <Description>Unreal Engine Multicast Delegate snippet</Description>
+          <HelpUrl>
+          </HelpUrl>
+          <Shortcut>ue4delmul6</Shortcut>
+        </Header>
+        <Snippet>
+          <Declarations>
+       <Literal Editable="true">
+          <ID>delegateName</ID>
+          <ToolTip>Delegate type name</ToolTip>
+          <Default>DelegateName</Default>
+          <Function>
+          </Function>
+        </Literal>
+        <Literal Editable="true">
+          <ID>paramName1</ID>
+          <ToolTip>Parameter name</ToolTip>
+          <Default>Param1Name</Default>
+          <Function>
+          </Function>
+            </Literal>
+        <Literal Editable="true">
+          <ID>paramName2</ID>
+          <ToolTip>Parameter name</ToolTip>
+          <Default>Param2Name</Default>
+          <Function>
+          </Function>
+            </Literal>
+        <Literal Editable="true">
+          <ID>paramName3</ID>
+          <ToolTip>Parameter name</ToolTip>
+          <Default>Param3Name</Default>
+          <Function>
+          </Function>
+            </Literal>
+        <Literal Editable="true">
+          <ID>paramName4</ID>
+          <ToolTip>Parameter name</ToolTip>
+          <Default>Param4Name</Default>
+          <Function>
+          </Function>
+            </Literal>
+        <Literal Editable="true">
+          <ID>paramName5</ID>
+          <ToolTip>Parameter name</ToolTip>
+          <Default>Param5Name</Default>
+          <Function>
+          </Function>
+            </Literal>
+        <Literal Editable="true">
+          <ID>paramName6</ID>
+          <ToolTip>Parameter name</ToolTip>
+          <Default>Param6Name</Default>
+          <Function>
+          </Function>
+            </Literal>
+       <Literal Editable="true">
+          <ID>paramType1</ID>
+          <ToolTip>Parameter type</ToolTip>
+          <Default>Param1Type</Default>
+          <Function>
+          </Function>
+            </Literal>
+       <Literal Editable="true">
+          <ID>paramType2</ID>
+          <ToolTip>Parameter type</ToolTip>
+          <Default>Param2Type</Default>
+          <Function>
+          </Function>
+            </Literal>
+       <Literal Editable="true">
+          <ID>paramType3</ID>
+          <ToolTip>Parameter type</ToolTip>
+          <Default>Param3Type</Default>
+          <Function>
+          </Function>
+            </Literal>
+       <Literal Editable="true">
+          <ID>paramType4</ID>
+          <ToolTip>Parameter type</ToolTip>
+          <Default>Param4Type</Default>
+          <Function>
+          </Function>
+            </Literal>
+       <Literal Editable="true">
+          <ID>paramType5</ID>
+          <ToolTip>Parameter type</ToolTip>
+          <Default>Param5Type</Default>
+          <Function>
+          </Function>
+            </Literal>
+       <Literal Editable="true">
+          <ID>paramType6</ID>
+          <ToolTip>Parameter type</ToolTip>
+          <Default>Param6Type</Default>
+          <Function>
+          </Function>
+            </Literal>
+        </Declarations>
+          <Code Language="cpp"><![CDATA[DECLARE_MULTICAST_DELEGATE_SixParams($delegateName$, $paramType1$, $paramName1$, $paramType2$, $paramName2$, $paramType3$, $paramName3$, $paramType4$, $paramName4$, $paramType5$, $paramName5$, $paramType6$, $paramName6$)$selected$ $end$]]></Code>
+        </Snippet>
+      </CodeSnippet>
+    </CodeSnippets>
\ No newline at end of file
diff --git a/Engine/Extras/VisualStudioSnippets/ue4delmul7.snippet b/Engine/Extras/VisualStudioSnippets/ue4delmul7.snippet
new file mode 100644
--- /dev/null
+++ b/Engine/Extras/VisualStudioSnippets/ue4delmul7.snippet
@@ -0,0 +1,126 @@
+<?xml version="1.0" encoding="utf-8"?>
+    <CodeSnippets xmlns="http://schemas.microsoft.com/VisualStudio/2005/CodeSnippet">
+      <CodeSnippet Format="1.0.0">
+        <Header>
+          <SnippetTypes>
+            <SnippetType>Expansion</SnippetType>
+          </SnippetTypes>
+          <Title>ue4delmul7</Title>
+          <Author>Epic Games</Author>
+          <Description>Unreal Engine Multicast Delegate snippet</Description>
+          <HelpUrl>
+          </HelpUrl>
+          <Shortcut>ue4delmul7</Shortcut>
+        </Header>
+        <Snippet>
+          <Declarations>
+       <Literal Editable="true">
+          <ID>delegateName</ID>
+          <ToolTip>Delegate type name</ToolTip>
+          <Default>DelegateName</Default>
+          <Function>
+          </Function>
+        </Literal>
+        <Literal Editable="true">
+          <ID>paramName1</ID>
+          <ToolTip>Parameter name</ToolTip>
+          <Default>Param1Name</Default>
+          <Function>
+          </Function>
+            </Literal>
+        <Literal Editable="true">
+          <ID>paramName2</ID>
+          <ToolTip>Parameter name</ToolTip>
+          <Default>Param2Name</Default>
+          <Function>
+          </Function>
+            </Literal>
+        <Literal Editable="true">
+          <ID>paramName3</ID>
+          <ToolTip>Parameter name</ToolTip>
+          <Default>Param3Name</Default>
+          <Function>
+          </Function>
+            </Literal>
+        <Literal Editable="true">
+          <ID>paramName4</ID>
+          <ToolTip>Parameter name</ToolTip>
+          <Default>Param4Name</Default>
+          <Function>
+          </Function>
+            </Literal>
+        <Literal Editable="true">
+          <ID>paramName5</ID>
+          <ToolTip>Parameter name</ToolTip>
+          <Default>Param5Name</Default>
+          <Function>
+          </Function>
+            </Literal>
+        <Literal Editable="true">
+          <ID>paramName6</ID>
+          <ToolTip>Parameter name</ToolTip>
+          <Default>Param6Name</Default>
+          <Function>
+          </Function>
+            </Literal>
+        <Literal Editable="true">
+          <ID>paramName7</ID>
+          <ToolTip>Parameter name</ToolTip>
+          <Default>Param7Name</Default>
+          <Function>
+          </Function>
+            </Literal>
+       <Literal Editable="true">
+          <ID>paramType1</ID>
+          <ToolTip>Parameter type</ToolTip>
+          <Default>Param1Type</Default>
+          <Function>
+          </Function>
+            </Literal>
+       <Literal Editable="true">
+          <ID>paramType2</ID>
+          <ToolTip>Parameter type</ToolTip>
+          <Default>Param2Type</Default>
+          <Function>
+          </Function>
+            </Literal>
+       <Literal Editable="true">
+          <ID>paramType3</ID>
+          <ToolTip>Parameter type</ToolTip>
+          <Default>Param3Type</Default>
+          <Function>
+          </Function>
+            </Literal>
+       <Literal Editable="true">
+          <ID>paramType4</ID>
+          <ToolTip>Parameter type</ToolTip>
+          <Default>Param4Type</Default>
+          <Function>
+          </Function>
+            </Literal>
+       <Literal Editable="true">
+          <ID>paramType5</ID>
+          <ToolTip>Parameter type</ToolTip>
+          <Default>Param5Type</Default>
+          <Function>
+          </Function>
+            </Literal>
+       <Literal Editable="true">
+          <ID>paramType6</ID>
+          <ToolTip>Parameter type</ToolTip>
+          <Default>Param6Type</Default>
+          <Function>
+          </Function>
+            </Literal>
+       <Literal Editable="true">
+          <ID>paramType7</ID>
+          <ToolTip>Parameter type</ToolTip>
+          <Default>Param7Type</Default>
+          <Function>
+          </Function>
+            </Literal>
+        </Declarations>
+          <Code Language="cpp"><![CDATA[DECLARE_MULTICAST_DELEGATE_SevenParams($delegateName$, $paramType1$, $paramName1$, $paramType2$, $paramName2$, $paramType3$, $paramName3$, $paramType4$, $paramName4$, $paramType5$, $paramName5$, $paramType6$, $paramName6$, $paramType7$, $paramName7$)$selected$ $end$]]></Code>
+        </Snippet>
+      </CodeSnippet>
+    </CodeSnippets>
\ No newline at end of file
diff --git a/Engine/Extras/VisualStudioSnippets/ue4delmul8.snippet b/Engine/Extras/VisualStudioSnippets/ue4delmul8.snippet
new file mode 100644
--- /dev/null
+++ b/Engine/Extras/VisualStudioSnippets/ue4delmul8.snippet
@@ -0,0 +1,140 @@
+<?xml version="1.0" encoding="utf-8"?>
+    <CodeSnippets xmlns="http://schemas.microsoft.com/VisualStudio/2005/CodeSnippet">
+      <CodeSnippet Format="1.0.0">
+        <Header>
+          <SnippetTypes>
+            <SnippetType>Expansion</SnippetType>
+          </SnippetTypes>
+          <Title>ue4delmul8</Title>
+          <Author>Epic Games</Author>
+          <Description>Unreal Engine Multicast Delegate snippet</Description>
+          <HelpUrl>
+          </HelpUrl>
+          <Shortcut>ue4delmul8</Shortcut>
+        </Header>
+        <Snippet>
+          <Declarations>
+       <Literal Editable="true">
+          <ID>delegateName</ID>
+          <ToolTip>Delegate type name</ToolTip>
+          <Default>DelegateName</Default>
+          <Function>
+          </Function>
+        </Literal>
+        <Literal Editable="true">
+          <ID>paramName1</ID>
+          <ToolTip>Parameter name</ToolTip>
+          <Default>Param1Name</Default>
+          <Function>
+          </Function>
+            </Literal>
+        <Literal Editable="true">
+          <ID>paramName2</ID>
+          <ToolTip>Parameter name</ToolTip>
+          <Default>Param2Name</Default>
+          <Function>
+          </Function>
+            </Literal>
+        <Literal Editable="true">
+          <ID>paramName3</ID>
+          <ToolTip>Parameter name</ToolTip>
+          <Default>Param3Name</Default>
+          <Function>
+          </Function>
+            </Literal>
+        <Literal Editable="true">
+          <ID>paramName4</ID>
+          <ToolTip>Parameter name</ToolTip>
+          <Default>Param4Name</Default>
+          <Function>
+          </Function>
+            </Literal>
+        <Literal Editable="true">
+          <ID>paramName5</ID>
+          <ToolTip>Parameter name</ToolTip>
+          <Default>Param5Name</Default>
+          <Function>
+          </Function>
+            </Literal>
+        <Literal Editable="true">
+          <ID>paramName6</ID>
+          <ToolTip>Parameter name</ToolTip>
+          <Default>Param6Name</Default>
+          <Function>
+          </Function>
+            </Literal>
+        <Literal Editable="true">
+          <ID>paramName7</ID>
+          <ToolTip>Parameter name</ToolTip>
+          <Default>Param7Name</Default>
+          <Function>
+          </Function>
+            </Literal>
+        <Literal Editable="true">
+          <ID>paramName8</ID>
+          <ToolTip>Parameter name</ToolTip>
+          <Default>Param8Name</Default>
+          <Function>
+          </Function>
+            </Literal>
+       <Literal Editable="true">
+          <ID>paramType1</ID>
+          <ToolTip>Parameter type</ToolTip>
+          <Default>Param1Type</Default>
+          <Function>
+          </Function>
+            </Literal>
+       <Literal Editable="true">
+          <ID>paramType2</ID>
+          <ToolTip>Parameter type</ToolTip>
+          <Default>Param2Type</Default>
+          <Function>
+          </Function>
+            </Literal>
+       <Literal Editable="true">
+          <ID>paramType3</ID>
+          <ToolTip>Parameter type</ToolTip>
+          <Default>Param3Type</Default>
+          <Function>
+          </Function>
+            </Literal>
+       <Literal Editable="true">
+          <ID>paramType4</ID>
+          <ToolTip>Parameter type</ToolTip>
+          <Default>Param4Type</Default>
+          <Function>
+          </Function>
+            </Literal>
+       <Literal Editable="true">
+          <ID>paramType5</ID>
+          <ToolTip>Parameter type</ToolTip>
+          <Default>Param5Type</Default>
+          <Function>
+          </Function>
+            </Literal>
+       <Literal Editable="true">
+          <ID>paramType6</ID>
+          <ToolTip>Parameter type</ToolTip>
+          <Default>Param6Type</Default>
+          <Function>
+          </Function>
+            </Literal>
+       <Literal Editable="true">
+          <ID>paramType7</ID>
+          <ToolTip>Parameter type</ToolTip>
+          <Default>Param7Type</Default>
+          <Function>
+          </Function>
+            </Literal>
+       <Literal Editable="true">
+          <ID>paramType8</ID>
+          <ToolTip>Parameter type</ToolTip>
+          <Default>Param8Type</Default>
+          <Function>
+          </Function>
+            </Literal>
+        </Declarations>
+          <Code Language="cpp"><![CDATA[DECLARE_MULTICAST_DELEGATE_EightParams($delegateName$, $paramType1$, $paramName1$, $paramType2$, $paramName2$, $paramType3$, $paramName3$, $paramType4$, $paramName4$, $paramType5$, $paramName5$, $paramType6$, $paramName6$, $paramType7$, $paramName7$, $paramType8$, $paramName8$)$selected$ $end$]]></Code>
+        </Snippet>
+      </CodeSnippet>
+    </CodeSnippets>
\ No newline at end of file
diff --git a/Engine/Extras/VisualStudioSnippets/ue4enum.snippet b/Engine/Extras/VisualStudioSnippets/ue4enum.snippet
new file mode 100644
--- /dev/null
+++ b/Engine/Extras/VisualStudioSnippets/ue4enum.snippet
@@ -0,0 +1,57 @@
+<?xml version="1.0" encoding="utf-8"?>
+<CodeSnippets xmlns="http://schemas.microsoft.com/VisualStudio/2005/CodeSnippet">
+  <CodeSnippet Format="1.0.0">
+    <Header>
+      <SnippetTypes>
+        <SnippetType>Expansion</SnippetType>
+      </SnippetTypes>
+      <Title>ue4enum</Title>
+      <Author>Epic Games</Author>
+      <Description>Unreal Engine Enum snippet</Description>
+      <HelpUrl>
+      </HelpUrl>
+      <Shortcut>ue4enum</Shortcut>
+    </Header>
+    <Snippet>
+      <Declarations>
+		<Literal Editable="true">
+          <ID>comment</ID>
+          <ToolTip>Enum comment</ToolTip>
+          <Default>Comment</Default>
+          <Function>
+          </Function>
+        </Literal>
+        <Literal Editable="true">
+          <ID>namespace</ID>
+          <ToolTip>Enum namespace name</ToolTip>
+          <Default>MyEnum</Default>
+          <Function>
+          </Function>
+        </Literal>
+		<Literal Editable="true">
+          <ID>elem1comment</ID>
+          <ToolTip>Enum element 1 comment</ToolTip>
+          <Default>Comment</Default>
+          <Function>
+          </Function>
+        </Literal>
+        <Literal Editable="true">
+          <ID>elem1name</ID>
+          <ToolTip>First enum element name</ToolTip>
+          <Default>Name</Default>
+          <Function>
+          </Function>
+        </Literal>
+      </Declarations>
+      <Code Language="cpp"><![CDATA[/**
+ * $comment$ 
+ */
+UENUM()
+enum E$namespace$
+{
+	$elem1name$, /**< $elem1comment$ */
+	$select$ $end$
+};]]></Code>
+    </Snippet>
+  </CodeSnippet>
+</CodeSnippets>
\ No newline at end of file
diff --git a/Engine/Extras/VisualStudioSnippets/ue4enumdisplay.snippet b/Engine/Extras/VisualStudioSnippets/ue4enumdisplay.snippet
new file mode 100644
--- /dev/null
+++ b/Engine/Extras/VisualStudioSnippets/ue4enumdisplay.snippet
@@ -0,0 +1,67 @@
+<?xml version="1.0" encoding="utf-8"?>
+<CodeSnippets xmlns="http://schemas.microsoft.com/VisualStudio/2005/CodeSnippet">
+  <CodeSnippet Format="1.0.0">
+    <Header>
+      <SnippetTypes>
+        <SnippetType>Expansion</SnippetType>
+      </SnippetTypes>
+      <Title>ue4enumdisplay</Title>
+      <Author>Epic Games</Author>
+      <Description>Unreal Engine Enum snippet</Description>
+      <HelpUrl>
+      </HelpUrl>
+      <Shortcut>ue4enumdisplay</Shortcut>
+    </Header>
+    <Snippet>
+      <Declarations>
+		<Literal Editable="true">
+          <ID>comment</ID>
+          <ToolTip>Enum comment</ToolTip>
+          <Default>Comment</Default>
+          <Function>
+          </Function>
+        </Literal>
+        <Literal Editable="true">
+          <ID>namespace</ID>
+          <ToolTip>Enum namespace name</ToolTip>
+          <Default>MyEnum</Default>
+          <Function>
+          </Function>
+        </Literal>
+        <Literal Editable="true">
+          <ID>elem1name</ID>
+          <ToolTip>First enum element name</ToolTip>
+          <Default>Name</Default>
+          <Function>
+          </Function>
+        </Literal>
+		<Literal Editable="true">
+          <ID>elem1comment</ID>
+          <ToolTip>Enum element 1 comment</ToolTip>
+          <Default>Comment</Default>
+          <Function>
+          </Function>
+        </Literal>
+		<Literal Editable="true">
+          <ID>elem1displayName</ID>
+          <ToolTip>First enum element display name</ToolTip>
+          <Default>DisplayName</Default>
+          <Function>
+          </Function>
+        </Literal>
+      </Declarations>
+      <Code Language="cpp"><![CDATA[/**
+ * $comment$ 
+ */
+UENUM()
+namespace E$namespace$
+{
+	enum Type
+	{
+		$elem1name$ UMETA(DisplayName="$elem1displayName$"), /**< $elem1comment$ */
+		$select$ $end$
+	};
+}]]></Code>
+    </Snippet>
+  </CodeSnippet>
+</CodeSnippets>
\ No newline at end of file
diff --git a/Engine/Extras/VisualStudioSnippets/ue4eve.snippet b/Engine/Extras/VisualStudioSnippets/ue4eve.snippet
new file mode 100644
--- /dev/null
+++ b/Engine/Extras/VisualStudioSnippets/ue4eve.snippet
@@ -0,0 +1,35 @@
+<?xml version="1.0" encoding="utf-8"?>
+    <CodeSnippets xmlns="http://schemas.microsoft.com/VisualStudio/2005/CodeSnippet">
+      <CodeSnippet Format="1.0.0">
+        <Header>
+          <SnippetTypes>
+            <SnippetType>Expansion</SnippetType>
+          </SnippetTypes>
+          <Title>ue4eve</Title>
+          <Author>Epic Games</Author>
+          <Description>Unreal Engine Event snippet</Description>
+          <HelpUrl>
+          </HelpUrl>
+          <Shortcut>ue4eve</Shortcut>
+        </Header>
+        <Snippet>
+          <Declarations>
+       <Literal Editable="true">
+          <ID>delegateName</ID>
+          <ToolTip>Delegate type name</ToolTip>
+          <Default>DelegateName</Default>
+          <Function>
+          </Function>
+        </Literal>
+       <Literal Editable="true">
+          <ID>owningType</ID>
+          <ToolTip>Delegate owning type name</ToolTip>
+          <Default>OwningType</Default>
+          <Function>
+          </Function>
+        </Literal>
+        </Declarations>
+          <Code Language="cpp"><![CDATA[DECLARE_EVENT($owningType$, $delegateName$)$selected$ $end$]]></Code>
+        </Snippet>
+      </CodeSnippet>
+    </CodeSnippets>
\ No newline at end of file
diff --git a/Engine/Extras/VisualStudioSnippets/ue4eve1.snippet b/Engine/Extras/VisualStudioSnippets/ue4eve1.snippet
new file mode 100644
--- /dev/null
+++ b/Engine/Extras/VisualStudioSnippets/ue4eve1.snippet
@@ -0,0 +1,42 @@
+<?xml version="1.0" encoding="utf-8"?>
+    <CodeSnippets xmlns="http://schemas.microsoft.com/VisualStudio/2005/CodeSnippet">
+      <CodeSnippet Format="1.0.0">
+        <Header>
+          <SnippetTypes>
+            <SnippetType>Expansion</SnippetType>
+          </SnippetTypes>
+          <Title>ue4eve1</Title>
+          <Author>Epic Games</Author>
+          <Description>Unreal Engine Event snippet</Description>
+          <HelpUrl>
+          </HelpUrl>
+          <Shortcut>ue4eve1</Shortcut>
+        </Header>
+        <Snippet>
+          <Declarations>
+       <Literal Editable="true">
+          <ID>delegateName</ID>
+          <ToolTip>Delegate type name</ToolTip>
+          <Default>DelegateName</Default>
+          <Function>
+          </Function>
+        </Literal>
+       <Literal Editable="true">
+          <ID>owningType</ID>
+          <ToolTip>Delegate owning type name</ToolTip>
+          <Default>OwningType</Default>
+          <Function>
+          </Function>
+        </Literal>
+       <Literal Editable="true">
+          <ID>paramType1</ID>
+          <ToolTip>Parameter type</ToolTip>
+          <Default>Param1Type</Default>
+          <Function>
+          </Function>
+            </Literal>
+        </Declarations>
+          <Code Language="cpp"><![CDATA[DECLARE_EVENT($owningType$, $delegateName$, $paramType1$)$selected$ $end$]]></Code>
+        </Snippet>
+      </CodeSnippet>
+    </CodeSnippets>
\ No newline at end of file
diff --git a/Engine/Extras/VisualStudioSnippets/ue4eve2.snippet b/Engine/Extras/VisualStudioSnippets/ue4eve2.snippet
new file mode 100644
--- /dev/null
+++ b/Engine/Extras/VisualStudioSnippets/ue4eve2.snippet
@@ -0,0 +1,49 @@
+<?xml version="1.0" encoding="utf-8"?>
+    <CodeSnippets xmlns="http://schemas.microsoft.com/VisualStudio/2005/CodeSnippet">
+      <CodeSnippet Format="1.0.0">
+        <Header>
+          <SnippetTypes>
+            <SnippetType>Expansion</SnippetType>
+          </SnippetTypes>
+          <Title>ue4eve2</Title>
+          <Author>Epic Games</Author>
+          <Description>Unreal Engine Event snippet</Description>
+          <HelpUrl>
+          </HelpUrl>
+          <Shortcut>ue4eve2</Shortcut>
+        </Header>
+        <Snippet>
+          <Declarations>
+       <Literal Editable="true">
+          <ID>delegateName</ID>
+          <ToolTip>Delegate type name</ToolTip>
+          <Default>DelegateName</Default>
+          <Function>
+          </Function>
+        </Literal>
+       <Literal Editable="true">
+          <ID>owningType</ID>
+          <ToolTip>Delegate owning type name</ToolTip>
+          <Default>OwningType</Default>
+          <Function>
+          </Function>
+        </Literal>
+       <Literal Editable="true">
+          <ID>paramType1</ID>
+          <ToolTip>Parameter type</ToolTip>
+          <Default>Param1Type</Default>
+          <Function>
+          </Function>
+            </Literal>
+       <Literal Editable="true">
+          <ID>paramType2</ID>
+          <ToolTip>Parameter type</ToolTip>
+          <Default>Param2Type</Default>
+          <Function>
+          </Function>
+            </Literal>
+        </Declarations>
+          <Code Language="cpp"><![CDATA[DECLARE_EVENT($owningType$, $delegateName$, $paramType1$, $paramType2$)$selected$ $end$]]></Code>
+        </Snippet>
+      </CodeSnippet>
+    </CodeSnippets>
\ No newline at end of file
diff --git a/Engine/Extras/VisualStudioSnippets/ue4eve3.snippet b/Engine/Extras/VisualStudioSnippets/ue4eve3.snippet
new file mode 100644
--- /dev/null
+++ b/Engine/Extras/VisualStudioSnippets/ue4eve3.snippet
@@ -0,0 +1,56 @@
+<?xml version="1.0" encoding="utf-8"?>
+    <CodeSnippets xmlns="http://schemas.microsoft.com/VisualStudio/2005/CodeSnippet">
+      <CodeSnippet Format="1.0.0">
+        <Header>
+          <SnippetTypes>
+            <SnippetType>Expansion</SnippetType>
+          </SnippetTypes>
+          <Title>ue4eve3</Title>
+          <Author>Epic Games</Author>
+          <Description>Unreal Engine Event snippet</Description>
+          <HelpUrl>
+          </HelpUrl>
+          <Shortcut>ue4eve3</Shortcut>
+        </Header>
+        <Snippet>
+          <Declarations>
+       <Literal Editable="true">
+          <ID>delegateName</ID>
+          <ToolTip>Delegate type name</ToolTip>
+          <Default>DelegateName</Default>
+          <Function>
+          </Function>
+        </Literal>
+       <Literal Editable="true">
+          <ID>owningType</ID>
+          <ToolTip>Delegate owning type name</ToolTip>
+          <Default>OwningType</Default>
+          <Function>
+          </Function>
+        </Literal>
+       <Literal Editable="true">
+          <ID>paramType1</ID>
+          <ToolTip>Parameter type</ToolTip>
+          <Default>Param1Type</Default>
+          <Function>
+          </Function>
+            </Literal>
+       <Literal Editable="true">
+          <ID>paramType2</ID>
+          <ToolTip>Parameter type</ToolTip>
+          <Default>Param2Type</Default>
+          <Function>
+          </Function>
+            </Literal>
+       <Literal Editable="true">
+          <ID>paramType3</ID>
+          <ToolTip>Parameter type</ToolTip>
+          <Default>Param3Type</Default>
+          <Function>
+          </Function>
+            </Literal>
+        </Declarations>
+          <Code Language="cpp"><![CDATA[DECLARE_EVENT($owningType$, $delegateName$, $paramType1$, $paramType2$, $paramType3$)$selected$ $end$]]></Code>
+        </Snippet>
+      </CodeSnippet>
+    </CodeSnippets>
\ No newline at end of file
diff --git a/Engine/Extras/VisualStudioSnippets/ue4eve4.snippet b/Engine/Extras/VisualStudioSnippets/ue4eve4.snippet
new file mode 100644
--- /dev/null
+++ b/Engine/Extras/VisualStudioSnippets/ue4eve4.snippet
@@ -0,0 +1,63 @@
+<?xml version="1.0" encoding="utf-8"?>
+    <CodeSnippets xmlns="http://schemas.microsoft.com/VisualStudio/2005/CodeSnippet">
+      <CodeSnippet Format="1.0.0">
+        <Header>
+          <SnippetTypes>
+            <SnippetType>Expansion</SnippetType>
+          </SnippetTypes>
+          <Title>ue4eve4</Title>
+          <Author>Epic Games</Author>
+          <Description>Unreal Engine Event snippet</Description>
+          <HelpUrl>
+          </HelpUrl>
+          <Shortcut>ue4eve4</Shortcut>
+        </Header>
+        <Snippet>
+          <Declarations>
+       <Literal Editable="true">
+          <ID>delegateName</ID>
+          <ToolTip>Delegate type name</ToolTip>
+          <Default>DelegateName</Default>
+          <Function>
+          </Function>
+        </Literal>
+       <Literal Editable="true">
+          <ID>owningType</ID>
+          <ToolTip>Delegate owning type name</ToolTip>
+          <Default>OwningType</Default>
+          <Function>
+          </Function>
+        </Literal>
+       <Literal Editable="true">
+          <ID>paramType1</ID>
+          <ToolTip>Parameter type</ToolTip>
+          <Default>Param1Type</Default>
+          <Function>
+          </Function>
+            </Literal>
+       <Literal Editable="true">
+          <ID>paramType2</ID>
+          <ToolTip>Parameter type</ToolTip>
+          <Default>Param2Type</Default>
+          <Function>
+          </Function>
+            </Literal>
+       <Literal Editable="true">
+          <ID>paramType3</ID>
+          <ToolTip>Parameter type</ToolTip>
+          <Default>Param3Type</Default>
+          <Function>
+          </Function>
+            </Literal>
+       <Literal Editable="true">
+          <ID>paramType4</ID>
+          <ToolTip>Parameter type</ToolTip>
+          <Default>Param4Type</Default>
+          <Function>
+          </Function>
+            </Literal>
+        </Declarations>
+          <Code Language="cpp"><![CDATA[DECLARE_EVENT($owningType$, $delegateName$, $paramType1$, $paramType2$, $paramType3$, $paramType4$)$selected$ $end$]]></Code>
+        </Snippet>
+      </CodeSnippet>
+    </CodeSnippets>
\ No newline at end of file
diff --git a/Engine/Extras/VisualStudioSnippets/ue4eve5.snippet b/Engine/Extras/VisualStudioSnippets/ue4eve5.snippet
new file mode 100644
--- /dev/null
+++ b/Engine/Extras/VisualStudioSnippets/ue4eve5.snippet
@@ -0,0 +1,70 @@
+<?xml version="1.0" encoding="utf-8"?>
+    <CodeSnippets xmlns="http://schemas.microsoft.com/VisualStudio/2005/CodeSnippet">
+      <CodeSnippet Format="1.0.0">
+        <Header>
+          <SnippetTypes>
+            <SnippetType>Expansion</SnippetType>
+          </SnippetTypes>
+          <Title>ue4eve5</Title>
+          <Author>Epic Games</Author>
+          <Description>Unreal Engine Event snippet</Description>
+          <HelpUrl>
+          </HelpUrl>
+          <Shortcut>ue4eve5</Shortcut>
+        </Header>
+        <Snippet>
+          <Declarations>
+       <Literal Editable="true">
+          <ID>delegateName</ID>
+          <ToolTip>Delegate type name</ToolTip>
+          <Default>DelegateName</Default>
+          <Function>
+          </Function>
+        </Literal>
+       <Literal Editable="true">
+          <ID>owningType</ID>
+          <ToolTip>Delegate owning type name</ToolTip>
+          <Default>OwningType</Default>
+          <Function>
+          </Function>
+        </Literal>
+       <Literal Editable="true">
+          <ID>paramType1</ID>
+          <ToolTip>Parameter type</ToolTip>
+          <Default>Param1Type</Default>
+          <Function>
+          </Function>
+            </Literal>
+       <Literal Editable="true">
+          <ID>paramType2</ID>
+          <ToolTip>Parameter type</ToolTip>
+          <Default>Param2Type</Default>
+          <Function>
+          </Function>
+            </Literal>
+       <Literal Editable="true">
+          <ID>paramType3</ID>
+          <ToolTip>Parameter type</ToolTip>
+          <Default>Param3Type</Default>
+          <Function>
+          </Function>
+            </Literal>
+       <Literal Editable="true">
+          <ID>paramType4</ID>
+          <ToolTip>Parameter type</ToolTip>
+          <Default>Param4Type</Default>
+          <Function>
+          </Function>
+            </Literal>
+       <Literal Editable="true">
+          <ID>paramType5</ID>
+          <ToolTip>Parameter type</ToolTip>
+          <Default>Param5Type</Default>
+          <Function>
+          </Function>
+            </Literal>
+        </Declarations>
+          <Code Language="cpp"><![CDATA[DECLARE_EVENT($owningType$, $delegateName$, $paramType1$, $paramType2$, $paramType3$, $paramType4$, $paramType5$)$selected$ $end$]]></Code>
+        </Snippet>
+      </CodeSnippet>
+    </CodeSnippets>
\ No newline at end of file
diff --git a/Engine/Extras/VisualStudioSnippets/ue4eve6.snippet b/Engine/Extras/VisualStudioSnippets/ue4eve6.snippet
new file mode 100644
--- /dev/null
+++ b/Engine/Extras/VisualStudioSnippets/ue4eve6.snippet
@@ -0,0 +1,77 @@
+<?xml version="1.0" encoding="utf-8"?>
+    <CodeSnippets xmlns="http://schemas.microsoft.com/VisualStudio/2005/CodeSnippet">
+      <CodeSnippet Format="1.0.0">
+        <Header>
+          <SnippetTypes>
+            <SnippetType>Expansion</SnippetType>
+          </SnippetTypes>
+          <Title>ue4eve6</Title>
+          <Author>Epic Games</Author>
+          <Description>Unreal Engine Event snippet</Description>
+          <HelpUrl>
+          </HelpUrl>
+          <Shortcut>ue4eve6</Shortcut>
+        </Header>
+        <Snippet>
+          <Declarations>
+       <Literal Editable="true">
+          <ID>delegateName</ID>
+          <ToolTip>Delegate type name</ToolTip>
+          <Default>DelegateName</Default>
+          <Function>
+          </Function>
+        </Literal>
+       <Literal Editable="true">
+          <ID>owningType</ID>
+          <ToolTip>Delegate owning type name</ToolTip>
+          <Default>OwningType</Default>
+          <Function>
+          </Function>
+        </Literal>
+       <Literal Editable="true">
+          <ID>paramType1</ID>
+          <ToolTip>Parameter type</ToolTip>
+          <Default>Param1Type</Default>
+          <Function>
+          </Function>
+            </Literal>
+       <Literal Editable="true">
+          <ID>paramType2</ID>
+          <ToolTip>Parameter type</ToolTip>
+          <Default>Param2Type</Default>
+          <Function>
+          </Function>
+            </Literal>
+       <Literal Editable="true">
+          <ID>paramType3</ID>
+          <ToolTip>Parameter type</ToolTip>
+          <Default>Param3Type</Default>
+          <Function>
+          </Function>
+            </Literal>
+       <Literal Editable="true">
+          <ID>paramType4</ID>
+          <ToolTip>Parameter type</ToolTip>
+          <Default>Param4Type</Default>
+          <Function>
+          </Function>
+            </Literal>
+       <Literal Editable="true">
+          <ID>paramType5</ID>
+          <ToolTip>Parameter type</ToolTip>
+          <Default>Param5Type</Default>
+          <Function>
+          </Function>
+            </Literal>
+       <Literal Editable="true">
+          <ID>paramType6</ID>
+          <ToolTip>Parameter type</ToolTip>
+          <Default>Param6Type</Default>
+          <Function>
+          </Function>
+            </Literal>
+        </Declarations>
+          <Code Language="cpp"><![CDATA[DECLARE_EVENT($owningType$, $delegateName$, $paramType1$, $paramType2$, $paramType3$, $paramType4$, $paramType5$, $paramType6$)$selected$ $end$]]></Code>
+        </Snippet>
+      </CodeSnippet>
+    </CodeSnippets>
\ No newline at end of file
diff --git a/Engine/Extras/VisualStudioSnippets/ue4eve7.snippet b/Engine/Extras/VisualStudioSnippets/ue4eve7.snippet
new file mode 100644
--- /dev/null
+++ b/Engine/Extras/VisualStudioSnippets/ue4eve7.snippet
@@ -0,0 +1,84 @@
+<?xml version="1.0" encoding="utf-8"?>
+    <CodeSnippets xmlns="http://schemas.microsoft.com/VisualStudio/2005/CodeSnippet">
+      <CodeSnippet Format="1.0.0">
+        <Header>
+          <SnippetTypes>
+            <SnippetType>Expansion</SnippetType>
+          </SnippetTypes>
+          <Title>ue4eve7</Title>
+          <Author>Epic Games</Author>
+          <Description>Unreal Engine Event snippet</Description>
+          <HelpUrl>
+          </HelpUrl>
+          <Shortcut>ue4eve7</Shortcut>
+        </Header>
+        <Snippet>
+          <Declarations>
+       <Literal Editable="true">
+          <ID>delegateName</ID>
+          <ToolTip>Delegate type name</ToolTip>
+          <Default>DelegateName</Default>
+          <Function>
+          </Function>
+        </Literal>
+       <Literal Editable="true">
+          <ID>owningType</ID>
+          <ToolTip>Delegate owning type name</ToolTip>
+          <Default>OwningType</Default>
+          <Function>
+          </Function>
+        </Literal>
+       <Literal Editable="true">
+          <ID>paramType1</ID>
+          <ToolTip>Parameter type</ToolTip>
+          <Default>Param1Type</Default>
+          <Function>
+          </Function>
+            </Literal>
+       <Literal Editable="true">
+          <ID>paramType2</ID>
+          <ToolTip>Parameter type</ToolTip>
+          <Default>Param2Type</Default>
+          <Function>
+          </Function>
+            </Literal>
+       <Literal Editable="true">
+          <ID>paramType3</ID>
+          <ToolTip>Parameter type</ToolTip>
+          <Default>Param3Type</Default>
+          <Function>
+          </Function>
+            </Literal>
+       <Literal Editable="true">
+          <ID>paramType4</ID>
+          <ToolTip>Parameter type</ToolTip>
+          <Default>Param4Type</Default>
+          <Function>
+          </Function>
+            </Literal>
+       <Literal Editable="true">
+          <ID>paramType5</ID>
+          <ToolTip>Parameter type</ToolTip>
+          <Default>Param5Type</Default>
+          <Function>
+          </Function>
+            </Literal>
+       <Literal Editable="true">
+          <ID>paramType6</ID>
+          <ToolTip>Parameter type</ToolTip>
+          <Default>Param6Type</Default>
+          <Function>
+          </Function>
+            </Literal>
+       <Literal Editable="true">
+          <ID>paramType7</ID>
+          <ToolTip>Parameter type</ToolTip>
+          <Default>Param7Type</Default>
+          <Function>
+          </Function>
+            </Literal>
+        </Declarations>
+          <Code Language="cpp"><![CDATA[DECLARE_EVENT($owningType$, $delegateName$, $paramType1$, $paramType2$, $paramType3$, $paramType4$, $paramType5$, $paramType6$, $paramType7$)$selected$ $end$]]></Code>
+        </Snippet>
+      </CodeSnippet>
+    </CodeSnippets>
\ No newline at end of file
diff --git a/Engine/Extras/VisualStudioSnippets/ue4eve8.snippet b/Engine/Extras/VisualStudioSnippets/ue4eve8.snippet
new file mode 100644
--- /dev/null
+++ b/Engine/Extras/VisualStudioSnippets/ue4eve8.snippet
@@ -0,0 +1,91 @@
+<?xml version="1.0" encoding="utf-8"?>
+    <CodeSnippets xmlns="http://schemas.microsoft.com/VisualStudio/2005/CodeSnippet">
+      <CodeSnippet Format="1.0.0">
+        <Header>
+          <SnippetTypes>
+            <SnippetType>Expansion</SnippetType>
+          </SnippetTypes>
+          <Title>ue4eve8</Title>
+          <Author>Epic Games</Author>
+          <Description>Unreal Engine Event snippet</Description>
+          <HelpUrl>
+          </HelpUrl>
+          <Shortcut>ue4eve8</Shortcut>
+        </Header>
+        <Snippet>
+          <Declarations>
+       <Literal Editable="true">
+          <ID>delegateName</ID>
+          <ToolTip>Delegate type name</ToolTip>
+          <Default>DelegateName</Default>
+          <Function>
+          </Function>
+        </Literal>
+       <Literal Editable="true">
+          <ID>owningType</ID>
+          <ToolTip>Delegate owning type name</ToolTip>
+          <Default>OwningType</Default>
+          <Function>
+          </Function>
+        </Literal>
+       <Literal Editable="true">
+          <ID>paramType1</ID>
+          <ToolTip>Parameter type</ToolTip>
+          <Default>Param1Type</Default>
+          <Function>
+          </Function>
+            </Literal>
+       <Literal Editable="true">
+          <ID>paramType2</ID>
+          <ToolTip>Parameter type</ToolTip>
+          <Default>Param2Type</Default>
+          <Function>
+          </Function>
+            </Literal>
+       <Literal Editable="true">
+          <ID>paramType3</ID>
+          <ToolTip>Parameter type</ToolTip>
+          <Default>Param3Type</Default>
+          <Function>
+          </Function>
+            </Literal>
+       <Literal Editable="true">
+          <ID>paramType4</ID>
+          <ToolTip>Parameter type</ToolTip>
+          <Default>Param4Type</Default>
+          <Function>
+          </Function>
+            </Literal>
+       <Literal Editable="true">
+          <ID>paramType5</ID>
+          <ToolTip>Parameter type</ToolTip>
+          <Default>Param5Type</Default>
+          <Function>
+          </Function>
+            </Literal>
+       <Literal Editable="true">
+          <ID>paramType6</ID>
+          <ToolTip>Parameter type</ToolTip>
+          <Default>Param6Type</Default>
+          <Function>
+          </Function>
+            </Literal>
+       <Literal Editable="true">
+          <ID>paramType7</ID>
+          <ToolTip>Parameter type</ToolTip>
+          <Default>Param7Type</Default>
+          <Function>
+          </Function>
+            </Literal>
+       <Literal Editable="true">
+          <ID>paramType8</ID>
+          <ToolTip>Parameter type</ToolTip>
+          <Default>Param8Type</Default>
+          <Function>
+          </Function>
+            </Literal>
+        </Declarations>
+          <Code Language="cpp"><![CDATA[DECLARE_EVENT($owningType$, $delegateName$, $paramType1$, $paramType2$, $paramType3$, $paramType4$, $paramType5$, $paramType6$, $paramType7$, $paramType8$)$selected$ $end$]]></Code>
+        </Snippet>
+      </CodeSnippet>
+    </CodeSnippets>
\ No newline at end of file
diff --git a/Engine/Extras/VisualStudioSnippets/ue4interface.snippet b/Engine/Extras/VisualStudioSnippets/ue4interface.snippet
new file mode 100644
--- /dev/null
+++ b/Engine/Extras/VisualStudioSnippets/ue4interface.snippet
@@ -0,0 +1,49 @@
+﻿<?xml version="1.0" encoding="utf-8"?>
+<CodeSnippets xmlns="http://schemas.microsoft.com/VisualStudio/2005/CodeSnippet">
+  <CodeSnippet Format="1.0.0">
+    <Header>
+      <SnippetTypes>
+        <SnippetType>Expansion</SnippetType>
+      </SnippetTypes>
+      <Title>ue4interface</Title>
+      <Author>Epic Games</Author>
+      <Description>Unreal Engine interface snippet</Description>
+      <HelpUrl>
+      </HelpUrl>
+      <Shortcut>ue4interface</Shortcut>
+    </Header>
+    <Snippet>
+      <Declarations>
+        <Literal Editable="true">
+          <ID>comment</ID>
+          <ToolTip>Interface comment</ToolTip>
+          <Default>Comment</Default>
+          <Function>
+          </Function>
+        </Literal>
+        <Literal Editable="true">
+          <ID>name</ID>
+          <ToolTip>Interface name</ToolTip>
+          <Default>Foo</Default>
+          <Function>
+          </Function>
+        </Literal>
+      </Declarations>
+      <Code Language="cpp"><![CDATA[/**
+ * $comment$
+ */
+UINTERFACE()
+class U$name$ : public UInterface
+{
+	GENERATED_UINTERFACE_BODY()
+};
+
+class I$name$
+{
+	GENERATED_IINTERFACE_BODY()
+  
+	$selected$ $end$
+};]]></Code>
+    </Snippet>
+  </CodeSnippet>
+</CodeSnippets>
\ No newline at end of file
diff --git a/Engine/Extras/VisualStudioSnippets/ue4log.snippet b/Engine/Extras/VisualStudioSnippets/ue4log.snippet
new file mode 100644
--- /dev/null
+++ b/Engine/Extras/VisualStudioSnippets/ue4log.snippet
@@ -0,0 +1,43 @@
+﻿<?xml version="1.0" encoding="utf-8"?>
+<CodeSnippets xmlns="http://schemas.microsoft.com/VisualStudio/2005/CodeSnippet">
+  <CodeSnippet Format="1.0.0">
+    <Header>
+      <SnippetTypes>
+        <SnippetType>Expansion</SnippetType>
+      </SnippetTypes>
+      <Title>ue4log</Title>
+      <Author>Epic Games</Author>
+      <Description>Unreal Engine UE_LOG snippet</Description>
+      <HelpUrl>
+      </HelpUrl>
+      <Shortcut>ue4log</Shortcut>
+    </Header>
+    <Snippet>
+      <Declarations>
+        <Literal Editable="true">
+          <ID>category</ID>
+          <ToolTip>Log category name</ToolTip>
+          <Default>MyProject</Default>
+          <Function>
+          </Function>
+        </Literal>
+        <Literal Editable="true">
+          <ID>level</ID>
+          <ToolTip>Log level (NoLoggigng, Fatal, Error, Warning, Display, Log, Verbose, VeryVerbose, All, NumVerbosity, VerbosityMask, SetColor, BreakOnLog)</ToolTip>
+          <Default>Error</Default>
+          <Function>
+          </Function>
+        </Literal>
+        <Literal Editable="true">
+          <ID>text</ID>
+          <ToolTip>Enter log message here</ToolTip>
+          <Default>Log message</Default>
+          <Function>
+          </Function>
+        </Literal>
+      </Declarations>
+      <Code Language="cpp"><![CDATA[UE_LOG($category$, $level$, TEXT("$text$"));
+$selected$ $end$]]></Code>
+    </Snippet>
+  </CodeSnippet>
+</CodeSnippets>
\ No newline at end of file
diff --git a/Engine/Extras/VisualStudioSnippets/ue4logdeclare.snippet b/Engine/Extras/VisualStudioSnippets/ue4logdeclare.snippet
new file mode 100644
--- /dev/null
+++ b/Engine/Extras/VisualStudioSnippets/ue4logdeclare.snippet
@@ -0,0 +1,43 @@
+﻿<?xml version="1.0" encoding="utf-8"?>
+<CodeSnippets xmlns="http://schemas.microsoft.com/VisualStudio/2005/CodeSnippet">
+  <CodeSnippet Format="1.0.0">
+    <Header>
+      <SnippetTypes>
+        <SnippetType>Expansion</SnippetType>
+      </SnippetTypes>
+      <Title>ue4logdeclare</Title>
+      <Author>Epic Games</Author>
+      <Description>Unreal Engine UE_LOG snippet</Description>
+      <HelpUrl>
+      </HelpUrl>
+      <Shortcut>ue4logdeclare</Shortcut>
+    </Header>
+    <Snippet>
+      <Declarations>
+        <Literal Editable="true">
+          <ID>category</ID>
+          <ToolTip>Log category name</ToolTip>
+          <Default>MyProject</Default>
+          <Function>
+          </Function>
+        </Literal>
+        <Literal Editable="true">
+          <ID>defaultVerbosity</ID>
+          <ToolTip>Default verbosity (NoLoggigng, Fatal, Error, Warning, Display, Log, Verbose, VeryVerbose, All, NumVerbosity, VerbosityMask, SetColor, BreakOnLog)</ToolTip>
+          <Default>Log</Default>
+          <Function>
+          </Function>
+        </Literal>
+        <Literal Editable="true">
+          <ID>compileTimeVerbosity</ID>
+          <ToolTip>Compile time verbosity</ToolTip>
+          <Default>All</Default>
+          <Function>
+          </Function>
+        </Literal>
+      </Declarations>
+      <Code Language="cpp"><![CDATA[DECLARE_LOG_CATEGORY_EXTERN($category$, $defaultVerbosity$, $compileTimeVerbosity$);
+$selected$ $end$]]></Code>
+    </Snippet>
+  </CodeSnippet>
+</CodeSnippets>
\ No newline at end of file
diff --git a/Engine/Extras/VisualStudioSnippets/ue4logdefine.snippet b/Engine/Extras/VisualStudioSnippets/ue4logdefine.snippet
new file mode 100644
--- /dev/null
+++ b/Engine/Extras/VisualStudioSnippets/ue4logdefine.snippet
@@ -0,0 +1,29 @@
+﻿<?xml version="1.0" encoding="utf-8"?>
+<CodeSnippets xmlns="http://schemas.microsoft.com/VisualStudio/2005/CodeSnippet">
+  <CodeSnippet Format="1.0.0">
+    <Header>
+      <SnippetTypes>
+        <SnippetType>Expansion</SnippetType>
+      </SnippetTypes>
+      <Title>ue4logdefine</Title>
+      <Author>Epic Games</Author>
+      <Description>Unreal Engine UE_LOG snippet</Description>
+      <HelpUrl>
+      </HelpUrl>
+      <Shortcut>ue4logdefine</Shortcut>
+    </Header>
+    <Snippet>
+      <Declarations>
+        <Literal Editable="true">
+          <ID>category</ID>
+          <ToolTip>Log category name</ToolTip>
+          <Default>MyProject</Default>
+          <Function>
+          </Function>
+        </Literal>
+      </Declarations>
+      <Code Language="cpp"><![CDATA[DEFINE_LOG_CATEGORY($category$);
+$selected$ $end$]]></Code>
+    </Snippet>
+  </CodeSnippet>
+</CodeSnippets>
\ No newline at end of file
diff --git a/Engine/Extras/VisualStudioSnippets/ue4logfloat.snippet b/Engine/Extras/VisualStudioSnippets/ue4logfloat.snippet
new file mode 100644
--- /dev/null
+++ b/Engine/Extras/VisualStudioSnippets/ue4logfloat.snippet
@@ -0,0 +1,43 @@
+﻿<?xml version="1.0" encoding="utf-8"?>
+<CodeSnippets xmlns="http://schemas.microsoft.com/VisualStudio/2005/CodeSnippet">
+  <CodeSnippet Format="1.0.0">
+    <Header>
+      <SnippetTypes>
+        <SnippetType>Expansion</SnippetType>
+      </SnippetTypes>
+      <Title>ue4logfloat</Title>
+      <Author>Epic Games</Author>
+      <Description>Unreal Engine UE_LOG snippet</Description>
+      <HelpUrl>
+      </HelpUrl>
+      <Shortcut>ue4logfloat</Shortcut>
+    </Header>
+    <Snippet>
+      <Declarations>
+        <Literal Editable="true">
+          <ID>category</ID>
+          <ToolTip>Log category name</ToolTip>
+          <Default>MyProject</Default>
+          <Function>
+          </Function>
+        </Literal>
+        <Literal Editable="true">
+          <ID>level</ID>
+          <ToolTip>Log level (NoLoggigng, Fatal, Error, Warning, Display, Log, Verbose, VeryVerbose, All, NumVerbosity, VerbosityMask, SetColor, BreakOnLog)</ToolTip>
+          <Default>Error</Default>
+          <Function>
+          </Function>
+        </Literal>
+        <Literal Editable="true">
+          <ID>variableName</ID>
+          <ToolTip>Name of variable</ToolTip>
+          <Default>variable</Default>
+          <Function>
+          </Function>
+        </Literal>
+      </Declarations>
+      <Code Language="cpp"><![CDATA[UE_LOG($category$, $level$, TEXT("The value of '$variableName$' is: %f"), $variableName$);
+$selected$ $end$]]></Code>
+    </Snippet>
+  </CodeSnippet>
+</CodeSnippets>
\ No newline at end of file
diff --git a/Engine/Extras/VisualStudioSnippets/ue4logint.snippet b/Engine/Extras/VisualStudioSnippets/ue4logint.snippet
new file mode 100644
--- /dev/null
+++ b/Engine/Extras/VisualStudioSnippets/ue4logint.snippet
@@ -0,0 +1,43 @@
+﻿<?xml version="1.0" encoding="utf-8"?>
+<CodeSnippets xmlns="http://schemas.microsoft.com/VisualStudio/2005/CodeSnippet">
+  <CodeSnippet Format="1.0.0">
+    <Header>
+      <SnippetTypes>
+        <SnippetType>Expansion</SnippetType>
+      </SnippetTypes>
+      <Title>ue4logint</Title>
+      <Author>Epic Games</Author>
+      <Description>Unreal Engine UE_LOG snippet</Description>
+      <HelpUrl>
+      </HelpUrl>
+      <Shortcut>ue4logint</Shortcut>
+    </Header>
+    <Snippet>
+      <Declarations>
+        <Literal Editable="true">
+          <ID>category</ID>
+          <ToolTip>Log category name</ToolTip>
+          <Default>MyProject</Default>
+          <Function>
+          </Function>
+        </Literal>
+        <Literal Editable="true">
+          <ID>level</ID>
+          <ToolTip>Log level (NoLoggigng, Fatal, Error, Warning, Display, Log, Verbose, VeryVerbose, All, NumVerbosity, VerbosityMask, SetColor, BreakOnLog)</ToolTip>
+          <Default>Error</Default>
+          <Function>
+          </Function>
+        </Literal>
+        <Literal Editable="true">
+          <ID>variableName</ID>
+          <ToolTip>Name of variable</ToolTip>
+          <Default>variable</Default>
+          <Function>
+          </Function>
+        </Literal>
+      </Declarations>
+      <Code Language="cpp"><![CDATA[UE_LOG($category$, $level$, TEXT("The value of '$variableName$' is: %i"), $variableName$);
+$selected$ $end$]]></Code>
+    </Snippet>
+  </CodeSnippet>
+</CodeSnippets>
\ No newline at end of file
diff --git a/Engine/Extras/VisualStudioSnippets/ue4loguobj.snippet b/Engine/Extras/VisualStudioSnippets/ue4loguobj.snippet
new file mode 100644
--- /dev/null
+++ b/Engine/Extras/VisualStudioSnippets/ue4loguobj.snippet
@@ -0,0 +1,50 @@
+﻿<?xml version="1.0" encoding="utf-8"?>
+<CodeSnippets xmlns="http://schemas.microsoft.com/VisualStudio/2005/CodeSnippet">
+  <CodeSnippet Format="1.0.0">
+    <Header>
+      <SnippetTypes>
+        <SnippetType>Expansion</SnippetType>
+      </SnippetTypes>
+      <Title>ue4loguobj</Title>
+      <Author>Epic Games</Author>
+      <Description>Unreal Engine UE_LOG snippet</Description>
+      <HelpUrl>
+      </HelpUrl>
+      <Shortcut>ue4loguobj</Shortcut>
+    </Header>
+    <Snippet>
+      <Declarations>
+        <Literal Editable="true">
+          <ID>category</ID>
+          <ToolTip>Log category name</ToolTip>
+          <Default>MyProject</Default>
+          <Function>
+          </Function>
+        </Literal>
+        <Literal Editable="true">
+          <ID>level</ID>
+          <ToolTip>Log level (NoLoggigng, Fatal, Error, Warning, Display, Log, Verbose, VeryVerbose, All, NumVerbosity, VerbosityMask, SetColor, BreakOnLog)</ToolTip>
+          <Default>Error</Default>
+          <Function>
+          </Function>
+        </Literal>
+        <Literal Editable="true">
+          <ID>text</ID>
+          <ToolTip>Enter log message here</ToolTip>
+          <Default>Log message</Default>
+          <Function>
+          </Function>
+        </Literal>
+		<Literal Editable="true">
+          <ID>object</ID>
+          <ToolTip>UObject object instance</ToolTip>
+          <Default>this</Default>
+          <Function>
+          </Function>
+        </Literal>
+      </Declarations>
+      <Code Language="cpp"><![CDATA[UE_LOG($category$, $level$, TEXT("[%s] $text$"), *($object$->GetName()));
+$selected$ $end$]]></Code>
+    </Snippet>
+  </CodeSnippet>
+</CodeSnippets>
\ No newline at end of file
diff --git a/Engine/Extras/VisualStudioSnippets/ue4mark.snippet b/Engine/Extras/VisualStudioSnippets/ue4mark.snippet
new file mode 100644
--- /dev/null
+++ b/Engine/Extras/VisualStudioSnippets/ue4mark.snippet
@@ -0,0 +1,62 @@
+﻿<?xml version="1.0" encoding="utf-8"?>
+<CodeSnippets xmlns="http://schemas.microsoft.com/VisualStudio/2005/CodeSnippet">
+  <CodeSnippet Format="1.0.0">
+    <Header>
+      <SnippetTypes>
+        <SnippetType>Expansion</SnippetType>
+      </SnippetTypes>
+      <Title>ue4mark</Title>
+      <Author>Epic Games</Author>
+      <Description>Unreal Engine mark snippet</Description>
+      <HelpUrl>
+      </HelpUrl>
+      <Shortcut>ue4mark</Shortcut>
+    </Header>
+    <Snippet>
+      <Declarations>
+        <Literal Editable="true">
+          <ID>companySymbol</ID>
+          <ToolTip>Symbol of your company</ToolTip>
+          <Default>MyCompany</Default>
+          <Function>
+          </Function>
+        </Literal>
+		<Literal Editable="true">
+          <ID>task</ID>
+          <ToolTip>Task/ticket number</ToolTip>
+          <Default>Ticket number</Default>
+          <Function>
+          </Function>
+        </Literal>
+		<Literal Editable="true">
+          <ID>name</ID>
+          <ToolTip>Developer name</ToolTip>
+          <Default>Name</Default>
+          <Function>
+          </Function>
+        </Literal>
+		<Literal Editable="true">
+          <ID>surname</ID>
+          <ToolTip>Developer surname</ToolTip>
+          <Default>Surname</Default>
+          <Function>
+          </Function>
+        </Literal>
+		<Literal Editable="true">
+          <ID>description</ID>
+          <ToolTip>Description of modification</ToolTip>
+          <Default>Description</Default>
+          <Function>
+          </Function>
+        </Literal>
+      </Declarations>
+      <Code Language="cpp"><![CDATA[/* BEGIN $companySymbol$ ($task$) */
+	// $name$ $surname$ - $description$
+
+	$selected$ $end$
+
+	/* END $companySymbol$ */
+]]></Code>
+    </Snippet>
+  </CodeSnippet>
+</CodeSnippets>
\ No newline at end of file
diff --git a/Engine/Extras/VisualStudioSnippets/ue4struct.snippet b/Engine/Extras/VisualStudioSnippets/ue4struct.snippet
new file mode 100644
--- /dev/null
+++ b/Engine/Extras/VisualStudioSnippets/ue4struct.snippet
@@ -0,0 +1,45 @@
+﻿<?xml version="1.0" encoding="utf-8"?>
+<CodeSnippets xmlns="http://schemas.microsoft.com/VisualStudio/2005/CodeSnippet">
+  <CodeSnippet Format="1.0.0">
+    <Header>
+      <SnippetTypes>
+        <SnippetType>Expansion</SnippetType>
+      </SnippetTypes>
+      <Title>ue4struct</Title>
+      <Author>Epic Games</Author>
+      <Description>Unreal Engine struct snippet</Description>
+      <HelpUrl>
+      </HelpUrl>
+      <Shortcut>ue4struct</Shortcut>
+    </Header>
+    <Snippet>
+      <Declarations>
+        <Literal Editable="true">
+          <ID>comment</ID>
+          <ToolTip>Struct comment</ToolTip>
+          <Default>Comment</Default>
+          <Function>
+          </Function>
+        </Literal>
+		<Literal Editable="true">
+          <ID>name</ID>
+          <ToolTip>Struct name</ToolTip>
+          <Default>Foo</Default>
+          <Function>
+          </Function>
+        </Literal>
+      </Declarations>
+      <Code Language="cpp"><![CDATA[/**
+ * $comment$ 
+ */
+USTRUCT()
+struct F$name$
+{
+	GENERATED_USTRUCT_BODY()
+
+	$selected$ $end$
+};
+]]></Code>
+    </Snippet>
+  </CodeSnippet>
+</CodeSnippets>
\ No newline at end of file
