Snippet Renames form Ue4 to Unreal

diff --git a/Engine/Extras/VisualStudioSnippets/README.md b/Engine/Extras/VisualStudioSnippets/README.md
--- a/Engine/Extras/VisualStudioSnippets/README.md
+++ b/Engine/Extras/VisualStudioSnippets/README.md
@@ -14,7 +14,7 @@ How to use snippets?
 ===========
 
 **Method One**
-Just start typing ue4... snippet list should be loaded in a form of combo box. Then use arrows to select snippet. Hit ENTER or TAB to insert snippet.
+Just start typing unreal-... snippet list should be loaded in a form of combo box. Then use arrows to select snippet. Hit ENTER or TAB to insert snippet.
 
 **Method Two**
 Type all snippet name and hit TAB. You don't have to wait for VS to show snippet list.
@@ -24,44 +24,44 @@ To navigate between highlighted fields you can use TAB and SHIFT + TAB. After yo
 Snippets
 ===========
 
-*ue4classa* – Blueprintable class that derives from an AActor. Parameters are: comment, class name and base class name.
+*unreal-classa* – Blueprintable class that derives from an AActor. Parameters are: comment, class name and base class name.
 
-*ue4classu* – Blueprintable class that derives from an UObject. Parameters are: comment, class name, base class name.
+*unreal-classu* – Blueprintable class that derives from an UObject. Parameters are: comment, class name, base class name.
 
-*ue4struct* – Simple structure. Parameters are: comment and name.
+*unreal-struct* – Simple structure. Parameters are: comment and name.
 
-*ue4interface* – Simple ue4 interface. Parameters are: comment and name.
+*unreal-interface* – Simple unreal- interface. Parameters are: comment and name.
 
-*ue4bpevent* – This function can be used as an event in blueprint. Parameters are: comment, UI category, virtual and const modifiers, function name and arguments.
+*unreal-bpevent* – This function can be used as an event in blueprint. Parameters are: comment, UI category, virtual and const modifiers, function name and arguments.
 
-*ue4bpfunc* – This function is available for blueprint logic. Parameters are: comment (parameters and return value), UI category, virtual and const modifiers, function name and arguments.
+*unreal-bpfunc* – This function is available for blueprint logic. Parameters are: comment (parameters and return value), UI category, virtual and const modifiers, function name and arguments.
 
-*ue4prop* – This read/write property is available everywhere (blueprint, instance and archetype details). Parameters are: comment, category, type and name.
+*unreal-prop* – This read/write property is available everywhere (blueprint, instance and archetype details). Parameters are: comment, category, type and name.
 
-*ue4enum* – Simple enum. Parameters are: comment, enum name, first member name and it’s comment.
+*unreal-enum* – Simple enum. Parameters are: comment, enum name, first member name and it’s comment.
 
-*ue4enumdisplay* – Enum that can be used with blueprints. Parameters are: comment, enum name, first member name, it’s display name and comment.
+*unreal-enumdisplay* – Enum that can be used with blueprints. Parameters are: comment, enum name, first member name, it’s display name and comment.
 
-*ue4log* – Simplest log line. Parameters are category, verbosity and message.
+*unreal-log* – Simplest log line. Parameters are category, verbosity and message.
 
-*ue4logdeclare* – Declaration of log category. Place this in main header of your project to allow logging. Parameters are: category, default verbosity and compile time verbosity.
+*unreal-logdeclare* – Declaration of log category. Place this in main header of your project to allow logging. Parameters are: category, default verbosity and compile time verbosity.
 
-*ue4logdefine* – Definition of log category. Place this in main code file. Parameter is category name.
+*unreal-logdefine* – Definition of log category. Place this in main code file. Parameter is category name.
 
-*ue4logfloat* – Log line that can be used to print float value. Parameters are: category, verbosity and variable name.
+*unreal-logfloat* – Log line that can be used to print float value. Parameters are: category, verbosity and variable name.
 
-*ue4logint* – This log line can be used to log an integer value. Parameters are: category, verbosity and variable name.
+*unreal-logint* – This log line can be used to log an integer value. Parameters are: category, verbosity and variable name.
 
-*ue4loguobj* – This log line is designed to log from inside of the objects. By default, square brackets contains a name of an object that writes the log. Parameters are: category, verbosity, message and name of a pointer to the object.
+*unreal-loguobj* – This log line is designed to log from inside of the objects. By default, square brackets contains a name of an object that writes the log. Parameters are: category, verbosity, message and name of a pointer to the object.
 
-*ue4mark* – Can be used to mark changes in engine classes. Parameters are: Company symbol, task/ticket number, name and surname of a developer and short description of modification.
+*unreal-mark* – Can be used to mark changes in engine classes. Parameters are: Company symbol, task/ticket number, name and surname of a developer and short description of modification.
 
-*ue4eve* - 9 snippets for each params combination. Can be used to create event. Parameters are: owning type and event type name.
+*unreal-eve* - 9 snippets for each params combination. Can be used to create event. Parameters are: owning type and event type name.
 
-*ue4del* - 9 snippets for each params combination. Can be used to create delegate. Parameters are: delegate type name and param type names.
+*unreal-del* - 9 snippets for each params combination. Can be used to create delegate. Parameters are: delegate type name and param type names.
 
-*ue4delmul* - 9 snippets for each params combination. Can be used to create multicast delegate. Parameters are: delegate type name and param type names.
+*unreal-delmul* - 9 snippets for each params combination. Can be used to create multicast delegate. Parameters are: delegate type name and param type names.
 
-*ue4deldyn* - 9 snippets for each params combination. Can be used to create dynamic delegate. Parameters are: delegate type name, param type names and display values.
+*unreal-deldyn* - 9 snippets for each params combination. Can be used to create dynamic delegate. Parameters are: delegate type name, param type names and display values.
 
-*ue4deldynmul* - 9 snippets for each params combination. Can be used to create dynamic multicast delegate. Parameters are: delegate type name, param type names and display values`.
\ No newline at end of file
+*unreal-deldynmul* - 9 snippets for each params combination. Can be used to create dynamic multicast delegate. Parameters are: delegate type name, param type names and display values`.
\ No newline at end of file
diff --git a/Engine/Extras/VisualStudioSnippets/UE4.SingleLineTrace.snippet b/Engine/Extras/VisualStudioSnippets/UE4.SingleLineTrace.snippet
deleted file mode 100644
diff --git a/Engine/Extras/VisualStudioSnippets/UE4.UFunction.snippet b/Engine/Extras/VisualStudioSnippets/UE4.UFunction.snippet
deleted file mode 100644
diff --git a/Engine/Extras/VisualStudioSnippets/UE4.UProperty.snippet b/Engine/Extras/VisualStudioSnippets/UE4.UProperty.snippet
deleted file mode 100644
diff --git a/Engine/Extras/VisualStudioSnippets/Unreal.SingleLineTrace.snippet b/Engine/Extras/VisualStudioSnippets/Unreal.SingleLineTrace.snippet
new file mode 100644
--- /dev/null
+++ b/Engine/Extras/VisualStudioSnippets/Unreal.SingleLineTrace.snippet
@@ -0,0 +1,26 @@
+<?xml version="1.0" encoding="utf-8" ?>
+<CodeSnippets xmlns="http://schemas.microsoft.com/VisualStudio/2005/CodeSnippet">
+   <CodeSnippet Format="1.0.0">
+      <Header>
+         <Title>unreal-SingleLineTrace</Title>
+         <Shortcut>SingleLineTrace</Shortcut>
+         <Description>Unreal Engine Single Line Trace (ignoring self) against World Dynamic</Description>
+         <Author>Epic Games</Author>
+         <SnippetTypes>
+            <SnippetType>Expansion</SnippetType>
+         </SnippetTypes>
+      </Header>
+      <Snippet>
+         <Code Language="cpp"><![CDATA[FCollisionQueryParams TraceParams(NAME_None, /*TraceComplex*/ true, /*IgnoreActor*/ this);
+		FCollisionObjectQueryParams CollisionParams(ECC_WorldDynamic);
+
+		FHitResult OutHit;
+		const bool bHit = World->LineTraceSingle(OutHit, StartPos, EndPos, TraceParams, CollisionParams);
+		if (bHit)
+		{
+			// OutHit Contains Hit Data...
+		}]]>
+         </Code>
+      </Snippet>
+   </CodeSnippet>
+</CodeSnippets>
\ No newline at end of file
diff --git a/Engine/Extras/VisualStudioSnippets/Unreal.UFunction.snippet b/Engine/Extras/VisualStudioSnippets/Unreal.UFunction.snippet
new file mode 100644
--- /dev/null
+++ b/Engine/Extras/VisualStudioSnippets/Unreal.UFunction.snippet
@@ -0,0 +1,21 @@
+<?xml version="1.0" encoding="utf-8" ?>
+<CodeSnippets xmlns="http://schemas.microsoft.com/VisualStudio/2005/CodeSnippet">
+   <CodeSnippet Format="1.0.0">
+      <Header>
+         <Title>unreal-UFunction</Title>
+         <Shortcut>ufunction</Shortcut>
+         <Description>Unreal Engine UFunction Template</Description>
+         <Author>Epic Games</Author>
+         <SnippetTypes>
+            <SnippetType>Expansion</SnippetType>
+         </SnippetTypes>
+      </Header>
+      <Snippet>
+         <Code Language="cpp"><![CDATA[/** Does the thing. */
+		UFUNCTION(BlueprintCallable, Category="Game Play")
+		void DoThing();
+		]]>
+         </Code>
+      </Snippet>
+   </CodeSnippet>
+</CodeSnippets>
\ No newline at end of file
diff --git a/Engine/Extras/VisualStudioSnippets/Unreal.UProperty.snippet b/Engine/Extras/VisualStudioSnippets/Unreal.UProperty.snippet
new file mode 100644
--- /dev/null
+++ b/Engine/Extras/VisualStudioSnippets/Unreal.UProperty.snippet
@@ -0,0 +1,21 @@
+<?xml version="1.0" encoding="utf-8" ?>
+<CodeSnippets xmlns="http://schemas.microsoft.com/VisualStudio/2005/CodeSnippet">
+   <CodeSnippet Format="1.0.0">
+      <Header>
+         <Title>unreal-UProperty</Title>
+         <Shortcut>uproperty</Shortcut>
+         <Description>Unreal Engine UProperty Template</Description>
+         <Author>Epic Games</Author>
+         <SnippetTypes>
+            <SnippetType>Expansion</SnippetType>
+         </SnippetTypes>
+      </Header>
+      <Snippet>
+         <Code Language="cpp"><![CDATA[/** The MyProperty */
+		UPROPERTY(EditAnywhere, BlueprintReadWrite, Category="Game")
+		float MyProperty;
+		]]>
+         </Code>
+      </Snippet>
+   </CodeSnippet>
+</CodeSnippets>
\ No newline at end of file
diff --git a/Engine/Extras/VisualStudioSnippets/Unreal.bpevent.snippet b/Engine/Extras/VisualStudioSnippets/Unreal.bpevent.snippet
new file mode 100644
--- /dev/null
+++ b/Engine/Extras/VisualStudioSnippets/Unreal.bpevent.snippet
@@ -0,0 +1,76 @@
+﻿<?xml version="1.0" encoding="utf-8"?>
+<CodeSnippets xmlns="http://schemas.microsoft.com/VisualStudio/2005/CodeSnippet">
+  <CodeSnippet Format="1.0.0">
+    <Header>
+      <SnippetTypes>
+        <SnippetType>Expansion</SnippetType>
+      </SnippetTypes>
+      <Title>unreal-bpevent</Title>
+      <Author>Epic Games</Author>
+      <Description>Unreal Engine blueprint implementable event snippet</Description>
+      <HelpUrl>
+      </HelpUrl>
+      <Shortcut>unreal-bpevent</Shortcut>
+    </Header>
+    <Snippet>
+      <Declarations>
+		<Literal Editable="true">
+          <ID>comment</ID>
+          <ToolTip>Function comment</ToolTip>
+          <Default>Comment</Default>
+          <Function>
+          </Function>
+        </Literal>
+		<Literal Editable="true">
+          <ID>commentParam</ID>
+          <ToolTip>Function comment</ToolTip>
+          <Default>Comment</Default>
+          <Function>
+          </Function>
+        </Literal>
+        <Literal Editable="true">
+          <ID>category</ID>
+          <ToolTip>UI category name</ToolTip>
+          <Default>MyProject</Default>
+          <Function>
+          </Function>
+        </Literal>
+		<Literal Editable="true">
+          <ID>virtual</ID>
+          <ToolTip>Virtual modifier</ToolTip>
+          <Default>virtual</Default>
+          <Function>
+          </Function>
+        </Literal>
+		<Literal Editable="true">
+          <ID>name</ID>
+          <ToolTip>Function name</ToolTip>
+          <Default>FunctionName</Default>
+          <Function>
+          </Function>
+        </Literal>
+		<Literal Editable="true">
+          <ID>args</ID>
+          <ToolTip>Arguments</ToolTip>
+          <Default>args</Default>
+          <Function>
+          </Function>
+        </Literal>
+		<Literal Editable="true">
+          <ID>const</ID>
+          <ToolTip>Const function modifier</ToolTip>
+          <Default>const</Default>
+          <Function>
+          </Function>
+        </Literal>
+      </Declarations>
+      <Code Language="cpp"><![CDATA[/**
+ * $comment$ 
+ * @param $commentParam$
+ */
+UFUNCTION(BlueprintImplementableEvent, Category="$category$")
+$virtual$ void On$name$($args$) $const$;
+$selected$ $end$]]></Code>
+    </Snippet>
+  </CodeSnippet>
+</CodeSnippets>
\ No newline at end of file
diff --git a/Engine/Extras/VisualStudioSnippets/Unreal.bpfunc.snippet b/Engine/Extras/VisualStudioSnippets/Unreal.bpfunc.snippet
new file mode 100644
--- /dev/null
+++ b/Engine/Extras/VisualStudioSnippets/Unreal.bpfunc.snippet
@@ -0,0 +1,91 @@
+﻿<?xml version="1.0" encoding="utf-8"?>
+<CodeSnippets xmlns="http://schemas.microsoft.com/VisualStudio/2005/CodeSnippet">
+  <CodeSnippet Format="1.0.0">
+    <Header>
+      <SnippetTypes>
+        <SnippetType>Expansion</SnippetType>
+      </SnippetTypes>
+      <Title>unreal-bpfunc</Title>
+      <Author>Epic Games</Author>
+      <Description>Unreal Engine blueprint callable function snippet</Description>
+      <HelpUrl>
+      </HelpUrl>
+      <Shortcut>unreal-bpfunc</Shortcut>
+    </Header>
+    <Snippet>
+      <Declarations>
+		<Literal Editable="true">
+          <ID>comment</ID>
+          <ToolTip>Function comment</ToolTip>
+          <Default>Comment</Default>
+          <Function>
+          </Function>
+        </Literal>
+		<Literal Editable="true">
+          <ID>commentParam</ID>
+          <ToolTip>Function comment</ToolTip>
+          <Default>Comment</Default>
+          <Function>
+          </Function>
+        </Literal>
+		<Literal Editable="true">
+          <ID>commentReturn</ID>
+          <ToolTip>Function comment</ToolTip>
+          <Default>Comment</Default>
+          <Function>
+          </Function>
+        </Literal>
+        <Literal Editable="true">
+          <ID>category</ID>
+          <ToolTip>UI category name</ToolTip>
+          <Default>MyProject</Default>
+          <Function>
+          </Function>
+        </Literal>
+		<Literal Editable="true">
+          <ID>virtual</ID>
+          <ToolTip>Virtual modifier</ToolTip>
+          <Default>virtual</Default>
+          <Function>
+          </Function>
+        </Literal>
+		<Literal Editable="true">
+          <ID>returnType</ID>
+          <ToolTip>Function return type</ToolTip>
+          <Default>void</Default>
+          <Function>
+          </Function>
+        </Literal>
+		<Literal Editable="true">
+          <ID>name</ID>
+          <ToolTip>Function name</ToolTip>
+          <Default>FunctionName</Default>
+          <Function>
+          </Function>
+        </Literal>
+		<Literal Editable="true">
+          <ID>args</ID>
+          <ToolTip>Arguments</ToolTip>
+          <Default>args</Default>
+          <Function>
+          </Function>
+        </Literal>
+		<Literal Editable="true">
+          <ID>const</ID>
+          <ToolTip>Const function modifier</ToolTip>
+          <Default>const</Default>
+          <Function>
+          </Function>
+        </Literal>
+      </Declarations>
+      <Code Language="cpp"><![CDATA[/**
+ * $comment$ 
+ * @param $commentParam$
+ * @return $commentReturn$
+ */
+UFUNCTION(BlueprintCallable, Category="$category$")
+$virtual$ $returnType$ $name$($args$) $const$;
+$selected$ $end$]]></Code>
+    </Snippet>
+  </CodeSnippet>
+</CodeSnippets>
\ No newline at end of file
diff --git a/Engine/Extras/VisualStudioSnippets/Unreal.bpprop.snippet b/Engine/Extras/VisualStudioSnippets/Unreal.bpprop.snippet
new file mode 100644
--- /dev/null
+++ b/Engine/Extras/VisualStudioSnippets/Unreal.bpprop.snippet
@@ -0,0 +1,68 @@
+﻿<?xml version="1.0" encoding="utf-8"?>
+<CodeSnippets xmlns="http://schemas.microsoft.com/VisualStudio/2005/CodeSnippet">
+  <CodeSnippet Format="1.0.0">
+    <Header>
+      <SnippetTypes>
+        <SnippetType>Expansion</SnippetType>
+      </SnippetTypes>
+      <Title>unreal-bpprop</Title>
+      <Author>Epic Games</Author>
+      <Description>Unreal Engine blueprint property snippet</Description>
+      <HelpUrl>
+      </HelpUrl>
+      <Shortcut>unreal-bpprop</Shortcut>
+    </Header>
+    <Snippet>
+      <Declarations>
+		<Literal Editable="true">
+          <ID>comment</ID>
+          <ToolTip>Property comment</ToolTip>
+          <Default>Comment</Default>
+          <Function>
+          </Function>
+        </Literal>
+        <Literal Editable="true">
+          <ID>category</ID>
+          <ToolTip>UI category name</ToolTip>
+          <Default>MyProject</Default>
+          <Function>
+          </Function>
+        </Literal>
+		<Literal Editable="true">
+          <ID>defaultsAttribute</ID>
+          <ToolTip>Defaults attribute</ToolTip>
+          <Default>EditAnywhere</Default>
+          <Function>
+          </Function>
+        </Literal>
+		<Literal Editable="true">
+          <ID>instanceAttribute</ID>
+          <ToolTip>Instance attribute</ToolTip>
+          <Default>BlueprintReadWrite</Default>
+          <Function>
+          </Function>
+        </Literal>
+		<Literal Editable="true">
+          <ID>type</ID>
+          <ToolTip>Property type</ToolTip>
+          <Default>Type</Default>
+          <Function>
+          </Function>
+        </Literal>
+		<Literal Editable="true">
+          <ID>name</ID>
+          <ToolTip>Property name</ToolTip>
+          <Default>Name</Default>
+          <Function>
+          </Function>
+        </Literal>
+      </Declarations>
+      <Code Language="cpp"><![CDATA[/** 
+ * $comment$ 
+ */
+UPROPERTY($defaultsAttribute$, $instanceAttribute$, Category="$category$")
+$type$ $name$;
+$selected$ $end$]]></Code>
+    </Snippet>
+  </CodeSnippet>
+</CodeSnippets>
\ No newline at end of file
diff --git a/Engine/Extras/VisualStudioSnippets/Unreal.classa.snippet b/Engine/Extras/VisualStudioSnippets/Unreal.classa.snippet
new file mode 100644
--- /dev/null
+++ b/Engine/Extras/VisualStudioSnippets/Unreal.classa.snippet
@@ -0,0 +1,58 @@
+﻿<?xml version="1.0" encoding="utf-8"?>
+<CodeSnippets xmlns="http://schemas.microsoft.com/VisualStudio/2005/CodeSnippet">
+  <CodeSnippet Format="1.0.0">
+    <Header>
+      <SnippetTypes>
+        <SnippetType>Expansion</SnippetType>
+      </SnippetTypes>
+      <Title>unreal-classa</Title>
+      <Author>Epic Games</Author>
+      <Description>Unreal Engine class snippet</Description>
+      <HelpUrl>
+      </HelpUrl>
+      <Shortcut>unreal-classa</Shortcut>
+    </Header>
+    <Snippet>
+      <Declarations>
+        <Literal Editable="true">
+          <ID>comment</ID>
+          <ToolTip>Class comment</ToolTip>
+          <Default>Comment</Default>
+          <Function>
+          </Function>
+        </Literal>
+		<Literal Editable="true">
+          <ID>className</ID>
+          <ToolTip>Class name</ToolTip>
+          <Default>Foo</Default>
+          <Function>
+          </Function>
+        </Literal>
+		<Literal Editable="true">
+          <ID>baseClassName</ID>
+          <ToolTip>Base class name</ToolTip>
+          <Default>FooBase</Default>
+          <Function>
+          </Function>
+        </Literal>
+      </Declarations>
+      <Code Language="cpp"><![CDATA[/**
+ * $comment$
+ */
+UCLASS(BlueprintType, Blueprintable)
+class A$className$ : public A$baseClassName$
+{
+	GENERATED_UCLASS_BODY()
+	
+public:
+	$selected$ $end$
+	
+protected:
+
+private:
+
+};
+]]></Code>
+    </Snippet>
+  </CodeSnippet>
+</CodeSnippets>
\ No newline at end of file
diff --git a/Engine/Extras/VisualStudioSnippets/Unreal.classu.snippet b/Engine/Extras/VisualStudioSnippets/Unreal.classu.snippet
new file mode 100644
--- /dev/null
+++ b/Engine/Extras/VisualStudioSnippets/Unreal.classu.snippet
@@ -0,0 +1,58 @@
+﻿<?xml version="1.0" encoding="utf-8"?>
+<CodeSnippets xmlns="http://schemas.microsoft.com/VisualStudio/2005/CodeSnippet">
+  <CodeSnippet Format="1.0.0">
+    <Header>
+      <SnippetTypes>
+        <SnippetType>Expansion</SnippetType>
+      </SnippetTypes>
+      <Title>unreal-classu</Title>
+      <Author>Epic Games</Author>
+      <Description>Unreal Engine class snippet</Description>
+      <HelpUrl>
+      </HelpUrl>
+      <Shortcut>unreal-classu</Shortcut>
+    </Header>
+    <Snippet>
+      <Declarations>
+        <Literal Editable="true">
+          <ID>comment</ID>
+          <ToolTip>Class comment</ToolTip>
+          <Default>Comment</Default>
+          <Function>
+          </Function>
+        </Literal>
+		<Literal Editable="true">
+          <ID>className</ID>
+          <ToolTip>Class name</ToolTip>
+          <Default>Foo</Default>
+          <Function>
+          </Function>
+        </Literal>
+		<Literal Editable="true">
+          <ID>baseClassName</ID>
+          <ToolTip>Base class name</ToolTip>
+          <Default>FooBase</Default>
+          <Function>
+          </Function>
+        </Literal>
+      </Declarations>
+      <Code Language="cpp"><![CDATA[/**
+ *$comment$
+ */
+UCLASS(BlueprintType, Blueprintable)
+class U$className$ : public U$baseClassName$
+{
+	GENERATED_UCLASS_BODY()
+
+public:
+	$selected$ $end$
+	
+protected:
+
+private:
+
+};
+]]></Code>
+    </Snippet>
+  </CodeSnippet>
+</CodeSnippets>
\ No newline at end of file
diff --git a/Engine/Extras/VisualStudioSnippets/Unreal.del.snippet b/Engine/Extras/VisualStudioSnippets/Unreal.del.snippet
new file mode 100644
--- /dev/null
+++ b/Engine/Extras/VisualStudioSnippets/Unreal.del.snippet
@@ -0,0 +1,28 @@
+<?xml version="1.0" encoding="utf-8"?>
+    <CodeSnippets xmlns="http://schemas.microsoft.com/VisualStudio/2005/CodeSnippet">
+      <CodeSnippet Format="1.0.0">
+        <Header>
+          <SnippetTypes>
+            <SnippetType>Expansion</SnippetType>
+          </SnippetTypes>
+          <Title>unreal-del</Title>
+          <Author>Epic Games</Author>
+          <Description>Unreal Engine Delegate snippet</Description>
+          <HelpUrl>
+          </HelpUrl>
+          <Shortcut>unreal-del</Shortcut>
+        </Header>
+        <Snippet>
+          <Declarations>
+       <Literal Editable="true">
+          <ID>delegateName</ID>
+          <ToolTip>Delegate type name</ToolTip>
+          <Default>DelegateName</Default>
+          <Function>
+          </Function>
+        </Literal>
+        </Declarations>
+          <Code Language="cpp"><![CDATA[DECLARE_DELEGATE($delegateName$)$selected$ $end$]]></Code>
+        </Snippet>
+      </CodeSnippet>
+    </CodeSnippets>
\ No newline at end of file
diff --git a/Engine/Extras/VisualStudioSnippets/Unreal.del1.snippet b/Engine/Extras/VisualStudioSnippets/Unreal.del1.snippet
new file mode 100644
--- /dev/null
+++ b/Engine/Extras/VisualStudioSnippets/Unreal.del1.snippet
@@ -0,0 +1,35 @@
+<?xml version="1.0" encoding="utf-8"?>
+    <CodeSnippets xmlns="http://schemas.microsoft.com/VisualStudio/2005/CodeSnippet">
+      <CodeSnippet Format="1.0.0">
+        <Header>
+          <SnippetTypes>
+            <SnippetType>Expansion</SnippetType>
+          </SnippetTypes>
+          <Title>unreal-del1</Title>
+          <Author>Epic Games</Author>
+          <Description>Unreal Engine Delegate snippet</Description>
+          <HelpUrl>
+          </HelpUrl>
+          <Shortcut>unreal-del1</Shortcut>
+        </Header>
+        <Snippet>
+          <Declarations>
+       <Literal Editable="true">
+          <ID>delegateName</ID>
+          <ToolTip>Delegate type name</ToolTip>
+          <Default>DelegateName</Default>
+          <Function>
+          </Function>
+        </Literal>
+       <Literal Editable="true">
+          <ID>paramType1</ID>
+          <ToolTip>Parameter type</ToolTip>
+          <Default>Param1Type</Default>
+          <Function>
+          </Function>
+            </Literal>
+        </Declarations>
+          <Code Language="cpp"><![CDATA[DECLARE_DELEGATE($delegateName$, $paramType1$)$selected$ $end$]]></Code>
+        </Snippet>
+      </CodeSnippet>
+    </CodeSnippets>
\ No newline at end of file
diff --git a/Engine/Extras/VisualStudioSnippets/Unreal.del2.snippet b/Engine/Extras/VisualStudioSnippets/Unreal.del2.snippet
new file mode 100644
--- /dev/null
+++ b/Engine/Extras/VisualStudioSnippets/Unreal.del2.snippet
@@ -0,0 +1,42 @@
+<?xml version="1.0" encoding="utf-8"?>
+    <CodeSnippets xmlns="http://schemas.microsoft.com/VisualStudio/2005/CodeSnippet">
+      <CodeSnippet Format="1.0.0">
+        <Header>
+          <SnippetTypes>
+            <SnippetType>Expansion</SnippetType>
+          </SnippetTypes>
+          <Title>unreal-del2</Title>
+          <Author>Epic Games</Author>
+          <Description>Unreal Engine Delegate snippet</Description>
+          <HelpUrl>
+          </HelpUrl>
+          <Shortcut>unreal-del2</Shortcut>
+        </Header>
+        <Snippet>
+          <Declarations>
+       <Literal Editable="true">
+          <ID>delegateName</ID>
+          <ToolTip>Delegate type name</ToolTip>
+          <Default>DelegateName</Default>
+          <Function>
+          </Function>
+        </Literal>
+       <Literal Editable="true">
+          <ID>paramType1</ID>
+          <ToolTip>Parameter type</ToolTip>
+          <Default>Param1Type</Default>
+          <Function>
+          </Function>
+            </Literal>
+       <Literal Editable="true">
+          <ID>paramType2</ID>
+          <ToolTip>Parameter type</ToolTip>
+          <Default>Param2Type</Default>
+          <Function>
+          </Function>
+            </Literal>
+        </Declarations>
+          <Code Language="cpp"><![CDATA[DECLARE_DELEGATE($delegateName$, $paramType1$, $paramType2$)$selected$ $end$]]></Code>
+        </Snippet>
+      </CodeSnippet>
+    </CodeSnippets>
\ No newline at end of file
diff --git a/Engine/Extras/VisualStudioSnippets/Unreal.del3.snippet b/Engine/Extras/VisualStudioSnippets/Unreal.del3.snippet
new file mode 100644
--- /dev/null
+++ b/Engine/Extras/VisualStudioSnippets/Unreal.del3.snippet
@@ -0,0 +1,49 @@
+<?xml version="1.0" encoding="utf-8"?>
+    <CodeSnippets xmlns="http://schemas.microsoft.com/VisualStudio/2005/CodeSnippet">
+      <CodeSnippet Format="1.0.0">
+        <Header>
+          <SnippetTypes>
+            <SnippetType>Expansion</SnippetType>
+          </SnippetTypes>
+          <Title>unreal-del3</Title>
+          <Author>Epic Games</Author>
+          <Description>Unreal Engine Delegate snippet</Description>
+          <HelpUrl>
+          </HelpUrl>
+          <Shortcut>unreal-del3</Shortcut>
+        </Header>
+        <Snippet>
+          <Declarations>
+       <Literal Editable="true">
+          <ID>delegateName</ID>
+          <ToolTip>Delegate type name</ToolTip>
+          <Default>DelegateName</Default>
+          <Function>
+          </Function>
+        </Literal>
+       <Literal Editable="true">
+          <ID>paramType1</ID>
+          <ToolTip>Parameter type</ToolTip>
+          <Default>Param1Type</Default>
+          <Function>
+          </Function>
+            </Literal>
+       <Literal Editable="true">
+          <ID>paramType2</ID>
+          <ToolTip>Parameter type</ToolTip>
+          <Default>Param2Type</Default>
+          <Function>
+          </Function>
+            </Literal>
+       <Literal Editable="true">
+          <ID>paramType3</ID>
+          <ToolTip>Parameter type</ToolTip>
+          <Default>Param3Type</Default>
+          <Function>
+          </Function>
+            </Literal>
+        </Declarations>
+          <Code Language="cpp"><![CDATA[DECLARE_DELEGATE($delegateName$, $paramType1$, $paramType2$, $paramType3$)$selected$ $end$]]></Code>
+        </Snippet>
+      </CodeSnippet>
+    </CodeSnippets>
\ No newline at end of file
diff --git a/Engine/Extras/VisualStudioSnippets/Unreal.del4.snippet b/Engine/Extras/VisualStudioSnippets/Unreal.del4.snippet
new file mode 100644
--- /dev/null
+++ b/Engine/Extras/VisualStudioSnippets/Unreal.del4.snippet
@@ -0,0 +1,56 @@
+<?xml version="1.0" encoding="utf-8"?>
+    <CodeSnippets xmlns="http://schemas.microsoft.com/VisualStudio/2005/CodeSnippet">
+      <CodeSnippet Format="1.0.0">
+        <Header>
+          <SnippetTypes>
+            <SnippetType>Expansion</SnippetType>
+          </SnippetTypes>
+          <Title>unreal-del4</Title>
+          <Author>Epic Games</Author>
+          <Description>Unreal Engine Delegate snippet</Description>
+          <HelpUrl>
+          </HelpUrl>
+          <Shortcut>unreal-del4</Shortcut>
+        </Header>
+        <Snippet>
+          <Declarations>
+       <Literal Editable="true">
+          <ID>delegateName</ID>
+          <ToolTip>Delegate type name</ToolTip>
+          <Default>DelegateName</Default>
+          <Function>
+          </Function>
+        </Literal>
+       <Literal Editable="true">
+          <ID>paramType1</ID>
+          <ToolTip>Parameter type</ToolTip>
+          <Default>Param1Type</Default>
+          <Function>
+          </Function>
+            </Literal>
+       <Literal Editable="true">
+          <ID>paramType2</ID>
+          <ToolTip>Parameter type</ToolTip>
+          <Default>Param2Type</Default>
+          <Function>
+          </Function>
+            </Literal>
+       <Literal Editable="true">
+          <ID>paramType3</ID>
+          <ToolTip>Parameter type</ToolTip>
+          <Default>Param3Type</Default>
+          <Function>
+          </Function>
+            </Literal>
+       <Literal Editable="true">
+          <ID>paramType4</ID>
+          <ToolTip>Parameter type</ToolTip>
+          <Default>Param4Type</Default>
+          <Function>
+          </Function>
+            </Literal>
+        </Declarations>
+          <Code Language="cpp"><![CDATA[DECLARE_DELEGATE($delegateName$, $paramType1$, $paramType2$, $paramType3$, $paramType4$)$selected$ $end$]]></Code>
+        </Snippet>
+      </CodeSnippet>
+    </CodeSnippets>
\ No newline at end of file
diff --git a/Engine/Extras/VisualStudioSnippets/Unreal.del5.snippet b/Engine/Extras/VisualStudioSnippets/Unreal.del5.snippet
new file mode 100644
--- /dev/null
+++ b/Engine/Extras/VisualStudioSnippets/Unreal.del5.snippet
@@ -0,0 +1,63 @@
+<?xml version="1.0" encoding="utf-8"?>
+    <CodeSnippets xmlns="http://schemas.microsoft.com/VisualStudio/2005/CodeSnippet">
+      <CodeSnippet Format="1.0.0">
+        <Header>
+          <SnippetTypes>
+            <SnippetType>Expansion</SnippetType>
+          </SnippetTypes>
+          <Title>unreal-del5</Title>
+          <Author>Epic Games</Author>
+          <Description>Unreal Engine Delegate snippet</Description>
+          <HelpUrl>
+          </HelpUrl>
+          <Shortcut>unreal-del5</Shortcut>
+        </Header>
+        <Snippet>
+          <Declarations>
+       <Literal Editable="true">
+          <ID>delegateName</ID>
+          <ToolTip>Delegate type name</ToolTip>
+          <Default>DelegateName</Default>
+          <Function>
+          </Function>
+        </Literal>
+       <Literal Editable="true">
+          <ID>paramType1</ID>
+          <ToolTip>Parameter type</ToolTip>
+          <Default>Param1Type</Default>
+          <Function>
+          </Function>
+            </Literal>
+       <Literal Editable="true">
+          <ID>paramType2</ID>
+          <ToolTip>Parameter type</ToolTip>
+          <Default>Param2Type</Default>
+          <Function>
+          </Function>
+            </Literal>
+       <Literal Editable="true">
+          <ID>paramType3</ID>
+          <ToolTip>Parameter type</ToolTip>
+          <Default>Param3Type</Default>
+          <Function>
+          </Function>
+            </Literal>
+       <Literal Editable="true">
+          <ID>paramType4</ID>
+          <ToolTip>Parameter type</ToolTip>
+          <Default>Param4Type</Default>
+          <Function>
+          </Function>
+            </Literal>
+       <Literal Editable="true">
+          <ID>paramType5</ID>
+          <ToolTip>Parameter type</ToolTip>
+          <Default>Param5Type</Default>
+          <Function>
+          </Function>
+            </Literal>
+        </Declarations>
+          <Code Language="cpp"><![CDATA[DECLARE_DELEGATE($delegateName$, $paramType1$, $paramType2$, $paramType3$, $paramType4$, $paramType5$)$selected$ $end$]]></Code>
+        </Snippet>
+      </CodeSnippet>
+    </CodeSnippets>
\ No newline at end of file
diff --git a/Engine/Extras/VisualStudioSnippets/Unreal.del6.snippet b/Engine/Extras/VisualStudioSnippets/Unreal.del6.snippet
new file mode 100644
--- /dev/null
+++ b/Engine/Extras/VisualStudioSnippets/Unreal.del6.snippet
@@ -0,0 +1,70 @@
+<?xml version="1.0" encoding="utf-8"?>
+    <CodeSnippets xmlns="http://schemas.microsoft.com/VisualStudio/2005/CodeSnippet">
+      <CodeSnippet Format="1.0.0">
+        <Header>
+          <SnippetTypes>
+            <SnippetType>Expansion</SnippetType>
+          </SnippetTypes>
+          <Title>unreal-del6</Title>
+          <Author>Epic Games</Author>
+          <Description>Unreal Engine Delegate snippet</Description>
+          <HelpUrl>
+          </HelpUrl>
+          <Shortcut>unreal-del6</Shortcut>
+        </Header>
+        <Snippet>
+          <Declarations>
+       <Literal Editable="true">
+          <ID>delegateName</ID>
+          <ToolTip>Delegate type name</ToolTip>
+          <Default>DelegateName</Default>
+          <Function>
+          </Function>
+        </Literal>
+       <Literal Editable="true">
+          <ID>paramType1</ID>
+          <ToolTip>Parameter type</ToolTip>
+          <Default>Param1Type</Default>
+          <Function>
+          </Function>
+            </Literal>
+       <Literal Editable="true">
+          <ID>paramType2</ID>
+          <ToolTip>Parameter type</ToolTip>
+          <Default>Param2Type</Default>
+          <Function>
+          </Function>
+            </Literal>
+       <Literal Editable="true">
+          <ID>paramType3</ID>
+          <ToolTip>Parameter type</ToolTip>
+          <Default>Param3Type</Default>
+          <Function>
+          </Function>
+            </Literal>
+       <Literal Editable="true">
+          <ID>paramType4</ID>
+          <ToolTip>Parameter type</ToolTip>
+          <Default>Param4Type</Default>
+          <Function>
+          </Function>
+            </Literal>
+       <Literal Editable="true">
+          <ID>paramType5</ID>
+          <ToolTip>Parameter type</ToolTip>
+          <Default>Param5Type</Default>
+          <Function>
+          </Function>
+            </Literal>
+       <Literal Editable="true">
+          <ID>paramType6</ID>
+          <ToolTip>Parameter type</ToolTip>
+          <Default>Param6Type</Default>
+          <Function>
+          </Function>
+            </Literal>
+        </Declarations>
+          <Code Language="cpp"><![CDATA[DECLARE_DELEGATE($delegateName$, $paramType1$, $paramType2$, $paramType3$, $paramType4$, $paramType5$, $paramType6$)$selected$ $end$]]></Code>
+        </Snippet>
+      </CodeSnippet>
+    </CodeSnippets>
\ No newline at end of file
diff --git a/Engine/Extras/VisualStudioSnippets/Unreal.del7.snippet b/Engine/Extras/VisualStudioSnippets/Unreal.del7.snippet
new file mode 100644
--- /dev/null
+++ b/Engine/Extras/VisualStudioSnippets/Unreal.del7.snippet
@@ -0,0 +1,77 @@
+<?xml version="1.0" encoding="utf-8"?>
+    <CodeSnippets xmlns="http://schemas.microsoft.com/VisualStudio/2005/CodeSnippet">
+      <CodeSnippet Format="1.0.0">
+        <Header>
+          <SnippetTypes>
+            <SnippetType>Expansion</SnippetType>
+          </SnippetTypes>
+          <Title>unreal-del7</Title>
+          <Author>Epic Games</Author>
+          <Description>Unreal Engine Delegate snippet</Description>
+          <HelpUrl>
+          </HelpUrl>
+          <Shortcut>unreal-del7</Shortcut>
+        </Header>
+        <Snippet>
+          <Declarations>
+       <Literal Editable="true">
+          <ID>delegateName</ID>
+          <ToolTip>Delegate type name</ToolTip>
+          <Default>DelegateName</Default>
+          <Function>
+          </Function>
+        </Literal>
+       <Literal Editable="true">
+          <ID>paramType1</ID>
+          <ToolTip>Parameter type</ToolTip>
+          <Default>Param1Type</Default>
+          <Function>
+          </Function>
+            </Literal>
+       <Literal Editable="true">
+          <ID>paramType2</ID>
+          <ToolTip>Parameter type</ToolTip>
+          <Default>Param2Type</Default>
+          <Function>
+          </Function>
+            </Literal>
+       <Literal Editable="true">
+          <ID>paramType3</ID>
+          <ToolTip>Parameter type</ToolTip>
+          <Default>Param3Type</Default>
+          <Function>
+          </Function>
+            </Literal>
+       <Literal Editable="true">
+          <ID>paramType4</ID>
+          <ToolTip>Parameter type</ToolTip>
+          <Default>Param4Type</Default>
+          <Function>
+          </Function>
+            </Literal>
+       <Literal Editable="true">
+          <ID>paramType5</ID>
+          <ToolTip>Parameter type</ToolTip>
+          <Default>Param5Type</Default>
+          <Function>
+          </Function>
+            </Literal>
+       <Literal Editable="true">
+          <ID>paramType6</ID>
+          <ToolTip>Parameter type</ToolTip>
+          <Default>Param6Type</Default>
+          <Function>
+          </Function>
+            </Literal>
+       <Literal Editable="true">
+          <ID>paramType7</ID>
+          <ToolTip>Parameter type</ToolTip>
+          <Default>Param7Type</Default>
+          <Function>
+          </Function>
+            </Literal>
+        </Declarations>
+          <Code Language="cpp"><![CDATA[DECLARE_DELEGATE($delegateName$, $paramType1$, $paramType2$, $paramType3$, $paramType4$, $paramType5$, $paramType6$, $paramType7$)$selected$ $end$]]></Code>
+        </Snippet>
+      </CodeSnippet>
+    </CodeSnippets>
\ No newline at end of file
diff --git a/Engine/Extras/VisualStudioSnippets/Unreal.del8.snippet b/Engine/Extras/VisualStudioSnippets/Unreal.del8.snippet
new file mode 100644
--- /dev/null
+++ b/Engine/Extras/VisualStudioSnippets/Unreal.del8.snippet
@@ -0,0 +1,84 @@
+<?xml version="1.0" encoding="utf-8"?>
+    <CodeSnippets xmlns="http://schemas.microsoft.com/VisualStudio/2005/CodeSnippet">
+      <CodeSnippet Format="1.0.0">
+        <Header>
+          <SnippetTypes>
+            <SnippetType>Expansion</SnippetType>
+          </SnippetTypes>
+          <Title>unreal-del8</Title>
+          <Author>Epic Games</Author>
+          <Description>Unreal Engine Delegate snippet</Description>
+          <HelpUrl>
+          </HelpUrl>
+          <Shortcut>unreal-del8</Shortcut>
+        </Header>
+        <Snippet>
+          <Declarations>
+       <Literal Editable="true">
+          <ID>delegateName</ID>
+          <ToolTip>Delegate type name</ToolTip>
+          <Default>DelegateName</Default>
+          <Function>
+          </Function>
+        </Literal>
+       <Literal Editable="true">
+          <ID>paramType1</ID>
+          <ToolTip>Parameter type</ToolTip>
+          <Default>Param1Type</Default>
+          <Function>
+          </Function>
+            </Literal>
+       <Literal Editable="true">
+          <ID>paramType2</ID>
+          <ToolTip>Parameter type</ToolTip>
+          <Default>Param2Type</Default>
+          <Function>
+          </Function>
+            </Literal>
+       <Literal Editable="true">
+          <ID>paramType3</ID>
+          <ToolTip>Parameter type</ToolTip>
+          <Default>Param3Type</Default>
+          <Function>
+          </Function>
+            </Literal>
+       <Literal Editable="true">
+          <ID>paramType4</ID>
+          <ToolTip>Parameter type</ToolTip>
+          <Default>Param4Type</Default>
+          <Function>
+          </Function>
+            </Literal>
+       <Literal Editable="true">
+          <ID>paramType5</ID>
+          <ToolTip>Parameter type</ToolTip>
+          <Default>Param5Type</Default>
+          <Function>
+          </Function>
+            </Literal>
+       <Literal Editable="true">
+          <ID>paramType6</ID>
+          <ToolTip>Parameter type</ToolTip>
+          <Default>Param6Type</Default>
+          <Function>
+          </Function>
+            </Literal>
+       <Literal Editable="true">
+          <ID>paramType7</ID>
+          <ToolTip>Parameter type</ToolTip>
+          <Default>Param7Type</Default>
+          <Function>
+          </Function>
+            </Literal>
+       <Literal Editable="true">
+          <ID>paramType8</ID>
+          <ToolTip>Parameter type</ToolTip>
+          <Default>Param8Type</Default>
+          <Function>
+          </Function>
+            </Literal>
+        </Declarations>
+          <Code Language="cpp"><![CDATA[DECLARE_DELEGATE($delegateName$, $paramType1$, $paramType2$, $paramType3$, $paramType4$, $paramType5$, $paramType6$, $paramType7$, $paramType8$)$selected$ $end$]]></Code>
+        </Snippet>
+      </CodeSnippet>
+    </CodeSnippets>
\ No newline at end of file
diff --git a/Engine/Extras/VisualStudioSnippets/Unreal.deldyn.snippet b/Engine/Extras/VisualStudioSnippets/Unreal.deldyn.snippet
new file mode 100644
--- /dev/null
+++ b/Engine/Extras/VisualStudioSnippets/Unreal.deldyn.snippet
@@ -0,0 +1,28 @@
+<?xml version="1.0" encoding="utf-8"?>
+    <CodeSnippets xmlns="http://schemas.microsoft.com/VisualStudio/2005/CodeSnippet">
+      <CodeSnippet Format="1.0.0">
+        <Header>
+          <SnippetTypes>
+            <SnippetType>Expansion</SnippetType>
+          </SnippetTypes>
+          <Title>unreal-deldyn</Title>
+          <Author>Epic Games</Author>
+          <Description>Unreal Engine Dynamic Delegate snippet</Description>
+          <HelpUrl>
+          </HelpUrl>
+          <Shortcut>unreal-deldyn</Shortcut>
+        </Header>
+        <Snippet>
+          <Declarations>
+       <Literal Editable="true">
+          <ID>delegateName</ID>
+          <ToolTip>Delegate type name</ToolTip>
+          <Default>DelegateName</Default>
+          <Function>
+          </Function>
+        </Literal>
+        </Declarations>
+          <Code Language="cpp"><![CDATA[DECLARE_DYNAMIC_DELEGATE($delegateName$)$selected$ $end$]]></Code>
+        </Snippet>
+      </CodeSnippet>
+    </CodeSnippets>
\ No newline at end of file
diff --git a/Engine/Extras/VisualStudioSnippets/Unreal.deldyn1.snippet b/Engine/Extras/VisualStudioSnippets/Unreal.deldyn1.snippet
new file mode 100644
--- /dev/null
+++ b/Engine/Extras/VisualStudioSnippets/Unreal.deldyn1.snippet
@@ -0,0 +1,42 @@
+<?xml version="1.0" encoding="utf-8"?>
+    <CodeSnippets xmlns="http://schemas.microsoft.com/VisualStudio/2005/CodeSnippet">
+      <CodeSnippet Format="1.0.0">
+        <Header>
+          <SnippetTypes>
+            <SnippetType>Expansion</SnippetType>
+          </SnippetTypes>
+          <Title>unreal-deldyn1</Title>
+          <Author>Epic Games</Author>
+          <Description>Unreal Engine Dynamic Delegate snippet</Description>
+          <HelpUrl>
+          </HelpUrl>
+          <Shortcut>unreal-deldyn1</Shortcut>
+        </Header>
+        <Snippet>
+          <Declarations>
+       <Literal Editable="true">
+          <ID>delegateName</ID>
+          <ToolTip>Delegate type name</ToolTip>
+          <Default>DelegateName</Default>
+          <Function>
+          </Function>
+        </Literal>
+        <Literal Editable="true">
+          <ID>paramName1</ID>
+          <ToolTip>Parameter name</ToolTip>
+          <Default>Param1Name</Default>
+          <Function>
+          </Function>
+            </Literal>
+       <Literal Editable="true">
+          <ID>paramType1</ID>
+          <ToolTip>Parameter type</ToolTip>
+          <Default>Param1Type</Default>
+          <Function>
+          </Function>
+            </Literal>
+        </Declarations>
+          <Code Language="cpp"><![CDATA[DECLARE_DYNAMIC_DELEGATE_OneParam($delegateName$, $paramType1$, $paramName1$)$selected$ $end$]]></Code>
+        </Snippet>
+      </CodeSnippet>
+    </CodeSnippets>
\ No newline at end of file
diff --git a/Engine/Extras/VisualStudioSnippets/Unreal.deldyn2.snippet b/Engine/Extras/VisualStudioSnippets/Unreal.deldyn2.snippet
new file mode 100644
--- /dev/null
+++ b/Engine/Extras/VisualStudioSnippets/Unreal.deldyn2.snippet
@@ -0,0 +1,56 @@
+<?xml version="1.0" encoding="utf-8"?>
+    <CodeSnippets xmlns="http://schemas.microsoft.com/VisualStudio/2005/CodeSnippet">
+      <CodeSnippet Format="1.0.0">
+        <Header>
+          <SnippetTypes>
+            <SnippetType>Expansion</SnippetType>
+          </SnippetTypes>
+          <Title>unreal-deldyn2</Title>
+          <Author>Epic Games</Author>
+          <Description>Unreal Engine Dynamic Delegate snippet</Description>
+          <HelpUrl>
+          </HelpUrl>
+          <Shortcut>unreal-deldyn2</Shortcut>
+        </Header>
+        <Snippet>
+          <Declarations>
+       <Literal Editable="true">
+          <ID>delegateName</ID>
+          <ToolTip>Delegate type name</ToolTip>
+          <Default>DelegateName</Default>
+          <Function>
+          </Function>
+        </Literal>
+        <Literal Editable="true">
+          <ID>paramName1</ID>
+          <ToolTip>Parameter name</ToolTip>
+          <Default>Param1Name</Default>
+          <Function>
+          </Function>
+            </Literal>
+        <Literal Editable="true">
+          <ID>paramName2</ID>
+          <ToolTip>Parameter name</ToolTip>
+          <Default>Param2Name</Default>
+          <Function>
+          </Function>
+            </Literal>
+       <Literal Editable="true">
+          <ID>paramType1</ID>
+          <ToolTip>Parameter type</ToolTip>
+          <Default>Param1Type</Default>
+          <Function>
+          </Function>
+            </Literal>
+       <Literal Editable="true">
+          <ID>paramType2</ID>
+          <ToolTip>Parameter type</ToolTip>
+          <Default>Param2Type</Default>
+          <Function>
+          </Function>
+            </Literal>
+        </Declarations>
+          <Code Language="cpp"><![CDATA[DECLARE_DYNAMIC_DELEGATE_TwoParams($delegateName$, $paramType1$, $paramName1$, $paramType2$, $paramName2$)$selected$ $end$]]></Code>
+        </Snippet>
+      </CodeSnippet>
+    </CodeSnippets>
\ No newline at end of file
diff --git a/Engine/Extras/VisualStudioSnippets/Unreal.deldyn3.snippet b/Engine/Extras/VisualStudioSnippets/Unreal.deldyn3.snippet
new file mode 100644
--- /dev/null
+++ b/Engine/Extras/VisualStudioSnippets/Unreal.deldyn3.snippet
@@ -0,0 +1,70 @@
+<?xml version="1.0" encoding="utf-8"?>
+    <CodeSnippets xmlns="http://schemas.microsoft.com/VisualStudio/2005/CodeSnippet">
+      <CodeSnippet Format="1.0.0">
+        <Header>
+          <SnippetTypes>
+            <SnippetType>Expansion</SnippetType>
+          </SnippetTypes>
+          <Title>unreal-deldyn3</Title>
+          <Author>Epic Games</Author>
+          <Description>Unreal Engine Dynamic Delegate snippet</Description>
+          <HelpUrl>
+          </HelpUrl>
+          <Shortcut>unreal-deldyn3</Shortcut>
+        </Header>
+        <Snippet>
+          <Declarations>
+       <Literal Editable="true">
+          <ID>delegateName</ID>
+          <ToolTip>Delegate type name</ToolTip>
+          <Default>DelegateName</Default>
+          <Function>
+          </Function>
+        </Literal>
+        <Literal Editable="true">
+          <ID>paramName1</ID>
+          <ToolTip>Parameter name</ToolTip>
+          <Default>Param1Name</Default>
+          <Function>
+          </Function>
+            </Literal>
+        <Literal Editable="true">
+          <ID>paramName2</ID>
+          <ToolTip>Parameter name</ToolTip>
+          <Default>Param2Name</Default>
+          <Function>
+          </Function>
+            </Literal>
+        <Literal Editable="true">
+          <ID>paramName3</ID>
+          <ToolTip>Parameter name</ToolTip>
+          <Default>Param3Name</Default>
+          <Function>
+          </Function>
+            </Literal>
+       <Literal Editable="true">
+          <ID>paramType1</ID>
+          <ToolTip>Parameter type</ToolTip>
+          <Default>Param1Type</Default>
+          <Function>
+          </Function>
+            </Literal>
+       <Literal Editable="true">
+          <ID>paramType2</ID>
+          <ToolTip>Parameter type</ToolTip>
+          <Default>Param2Type</Default>
+          <Function>
+          </Function>
+            </Literal>
+       <Literal Editable="true">
+          <ID>paramType3</ID>
+          <ToolTip>Parameter type</ToolTip>
+          <Default>Param3Type</Default>
+          <Function>
+          </Function>
+            </Literal>
+        </Declarations>
+          <Code Language="cpp"><![CDATA[DECLARE_DYNAMIC_DELEGATE_ThreeParams($delegateName$, $paramType1$, $paramName1$, $paramType2$, $paramName2$, $paramType3$, $paramName3$)$selected$ $end$]]></Code>
+        </Snippet>
+      </CodeSnippet>
+    </CodeSnippets>
\ No newline at end of file
diff --git a/Engine/Extras/VisualStudioSnippets/Unreal.deldyn4.snippet b/Engine/Extras/VisualStudioSnippets/Unreal.deldyn4.snippet
new file mode 100644
--- /dev/null
+++ b/Engine/Extras/VisualStudioSnippets/Unreal.deldyn4.snippet
@@ -0,0 +1,84 @@
+<?xml version="1.0" encoding="utf-8"?>
+    <CodeSnippets xmlns="http://schemas.microsoft.com/VisualStudio/2005/CodeSnippet">
+      <CodeSnippet Format="1.0.0">
+        <Header>
+          <SnippetTypes>
+            <SnippetType>Expansion</SnippetType>
+          </SnippetTypes>
+          <Title>unreal-deldyn4</Title>
+          <Author>Epic Games</Author>
+          <Description>Unreal Engine Dynamic Delegate snippet</Description>
+          <HelpUrl>
+          </HelpUrl>
+          <Shortcut>unreal-deldyn4</Shortcut>
+        </Header>
+        <Snippet>
+          <Declarations>
+       <Literal Editable="true">
+          <ID>delegateName</ID>
+          <ToolTip>Delegate type name</ToolTip>
+          <Default>DelegateName</Default>
+          <Function>
+          </Function>
+        </Literal>
+        <Literal Editable="true">
+          <ID>paramName1</ID>
+          <ToolTip>Parameter name</ToolTip>
+          <Default>Param1Name</Default>
+          <Function>
+          </Function>
+            </Literal>
+        <Literal Editable="true">
+          <ID>paramName2</ID>
+          <ToolTip>Parameter name</ToolTip>
+          <Default>Param2Name</Default>
+          <Function>
+          </Function>
+            </Literal>
+        <Literal Editable="true">
+          <ID>paramName3</ID>
+          <ToolTip>Parameter name</ToolTip>
+          <Default>Param3Name</Default>
+          <Function>
+          </Function>
+            </Literal>
+        <Literal Editable="true">
+          <ID>paramName4</ID>
+          <ToolTip>Parameter name</ToolTip>
+          <Default>Param4Name</Default>
+          <Function>
+          </Function>
+            </Literal>
+       <Literal Editable="true">
+          <ID>paramType1</ID>
+          <ToolTip>Parameter type</ToolTip>
+          <Default>Param1Type</Default>
+          <Function>
+          </Function>
+            </Literal>
+       <Literal Editable="true">
+          <ID>paramType2</ID>
+          <ToolTip>Parameter type</ToolTip>
+          <Default>Param2Type</Default>
+          <Function>
+          </Function>
+            </Literal>
+       <Literal Editable="true">
+          <ID>paramType3</ID>
+          <ToolTip>Parameter type</ToolTip>
+          <Default>Param3Type</Default>
+          <Function>
+          </Function>
+            </Literal>
+       <Literal Editable="true">
+          <ID>paramType4</ID>
+          <ToolTip>Parameter type</ToolTip>
+          <Default>Param4Type</Default>
+          <Function>
+          </Function>
+            </Literal>
+        </Declarations>
+          <Code Language="cpp"><![CDATA[DECLARE_DYNAMIC_DELEGATE_FourParams($delegateName$, $paramType1$, $paramName1$, $paramType2$, $paramName2$, $paramType3$, $paramName3$, $paramType4$, $paramName4$)$selected$ $end$]]></Code>
+        </Snippet>
+      </CodeSnippet>
+    </CodeSnippets>
\ No newline at end of file
diff --git a/Engine/Extras/VisualStudioSnippets/Unreal.deldyn5.snippet b/Engine/Extras/VisualStudioSnippets/Unreal.deldyn5.snippet
new file mode 100644
--- /dev/null
+++ b/Engine/Extras/VisualStudioSnippets/Unreal.deldyn5.snippet
@@ -0,0 +1,98 @@
+<?xml version="1.0" encoding="utf-8"?>
+    <CodeSnippets xmlns="http://schemas.microsoft.com/VisualStudio/2005/CodeSnippet">
+      <CodeSnippet Format="1.0.0">
+        <Header>
+          <SnippetTypes>
+            <SnippetType>Expansion</SnippetType>
+          </SnippetTypes>
+          <Title>unreal-deldyn5</Title>
+          <Author>Epic Games</Author>
+          <Description>Unreal Engine Dynamic Delegate snippet</Description>
+          <HelpUrl>
+          </HelpUrl>
+          <Shortcut>unreal-deldyn5</Shortcut>
+        </Header>
+        <Snippet>
+          <Declarations>
+       <Literal Editable="true">
+          <ID>delegateName</ID>
+          <ToolTip>Delegate type name</ToolTip>
+          <Default>DelegateName</Default>
+          <Function>
+          </Function>
+        </Literal>
+        <Literal Editable="true">
+          <ID>paramName1</ID>
+          <ToolTip>Parameter name</ToolTip>
+          <Default>Param1Name</Default>
+          <Function>
+          </Function>
+            </Literal>
+        <Literal Editable="true">
+          <ID>paramName2</ID>
+          <ToolTip>Parameter name</ToolTip>
+          <Default>Param2Name</Default>
+          <Function>
+          </Function>
+            </Literal>
+        <Literal Editable="true">
+          <ID>paramName3</ID>
+          <ToolTip>Parameter name</ToolTip>
+          <Default>Param3Name</Default>
+          <Function>
+          </Function>
+            </Literal>
+        <Literal Editable="true">
+          <ID>paramName4</ID>
+          <ToolTip>Parameter name</ToolTip>
+          <Default>Param4Name</Default>
+          <Function>
+          </Function>
+            </Literal>
+        <Literal Editable="true">
+          <ID>paramName5</ID>
+          <ToolTip>Parameter name</ToolTip>
+          <Default>Param5Name</Default>
+          <Function>
+          </Function>
+            </Literal>
+       <Literal Editable="true">
+          <ID>paramType1</ID>
+          <ToolTip>Parameter type</ToolTip>
+          <Default>Param1Type</Default>
+          <Function>
+          </Function>
+            </Literal>
+       <Literal Editable="true">
+          <ID>paramType2</ID>
+          <ToolTip>Parameter type</ToolTip>
+          <Default>Param2Type</Default>
+          <Function>
+          </Function>
+            </Literal>
+       <Literal Editable="true">
+          <ID>paramType3</ID>
+          <ToolTip>Parameter type</ToolTip>
+          <Default>Param3Type</Default>
+          <Function>
+          </Function>
+            </Literal>
+       <Literal Editable="true">
+          <ID>paramType4</ID>
+          <ToolTip>Parameter type</ToolTip>
+          <Default>Param4Type</Default>
+          <Function>
+          </Function>
+            </Literal>
+       <Literal Editable="true">
+          <ID>paramType5</ID>
+          <ToolTip>Parameter type</ToolTip>
+          <Default>Param5Type</Default>
+          <Function>
+          </Function>
+            </Literal>
+        </Declarations>
+          <Code Language="cpp"><![CDATA[DECLARE_DYNAMIC_DELEGATE_FiveParams($delegateName$, $paramType1$, $paramName1$, $paramType2$, $paramName2$, $paramType3$, $paramName3$, $paramType4$, $paramName4$, $paramType5$, $paramName5$)$selected$ $end$]]></Code>
+        </Snippet>
+      </CodeSnippet>
+    </CodeSnippets>
\ No newline at end of file
diff --git a/Engine/Extras/VisualStudioSnippets/Unreal.deldyn6.snippet b/Engine/Extras/VisualStudioSnippets/Unreal.deldyn6.snippet
new file mode 100644
--- /dev/null
+++ b/Engine/Extras/VisualStudioSnippets/Unreal.deldyn6.snippet
@@ -0,0 +1,112 @@
+<?xml version="1.0" encoding="utf-8"?>
+    <CodeSnippets xmlns="http://schemas.microsoft.com/VisualStudio/2005/CodeSnippet">
+      <CodeSnippet Format="1.0.0">
+        <Header>
+          <SnippetTypes>
+            <SnippetType>Expansion</SnippetType>
+          </SnippetTypes>
+          <Title>unreal-deldyn6</Title>
+          <Author>Epic Games</Author>
+          <Description>Unreal Engine Dynamic Delegate snippet</Description>
+          <HelpUrl>
+          </HelpUrl>
+          <Shortcut>unreal-deldyn6</Shortcut>
+        </Header>
+        <Snippet>
+          <Declarations>
+       <Literal Editable="true">
+          <ID>delegateName</ID>
+          <ToolTip>Delegate type name</ToolTip>
+          <Default>DelegateName</Default>
+          <Function>
+          </Function>
+        </Literal>
+        <Literal Editable="true">
+          <ID>paramName1</ID>
+          <ToolTip>Parameter name</ToolTip>
+          <Default>Param1Name</Default>
+          <Function>
+          </Function>
+            </Literal>
+        <Literal Editable="true">
+          <ID>paramName2</ID>
+          <ToolTip>Parameter name</ToolTip>
+          <Default>Param2Name</Default>
+          <Function>
+          </Function>
+            </Literal>
+        <Literal Editable="true">
+          <ID>paramName3</ID>
+          <ToolTip>Parameter name</ToolTip>
+          <Default>Param3Name</Default>
+          <Function>
+          </Function>
+            </Literal>
+        <Literal Editable="true">
+          <ID>paramName4</ID>
+          <ToolTip>Parameter name</ToolTip>
+          <Default>Param4Name</Default>
+          <Function>
+          </Function>
+            </Literal>
+        <Literal Editable="true">
+          <ID>paramName5</ID>
+          <ToolTip>Parameter name</ToolTip>
+          <Default>Param5Name</Default>
+          <Function>
+          </Function>
+            </Literal>
+        <Literal Editable="true">
+          <ID>paramName6</ID>
+          <ToolTip>Parameter name</ToolTip>
+          <Default>Param6Name</Default>
+          <Function>
+          </Function>
+            </Literal>
+       <Literal Editable="true">
+          <ID>paramType1</ID>
+          <ToolTip>Parameter type</ToolTip>
+          <Default>Param1Type</Default>
+          <Function>
+          </Function>
+            </Literal>
+       <Literal Editable="true">
+          <ID>paramType2</ID>
+          <ToolTip>Parameter type</ToolTip>
+          <Default>Param2Type</Default>
+          <Function>
+          </Function>
+            </Literal>
+       <Literal Editable="true">
+          <ID>paramType3</ID>
+          <ToolTip>Parameter type</ToolTip>
+          <Default>Param3Type</Default>
+          <Function>
+          </Function>
+            </Literal>
+       <Literal Editable="true">
+          <ID>paramType4</ID>
+          <ToolTip>Parameter type</ToolTip>
+          <Default>Param4Type</Default>
+          <Function>
+          </Function>
+            </Literal>
+       <Literal Editable="true">
+          <ID>paramType5</ID>
+          <ToolTip>Parameter type</ToolTip>
+          <Default>Param5Type</Default>
+          <Function>
+          </Function>
+            </Literal>
+       <Literal Editable="true">
+          <ID>paramType6</ID>
+          <ToolTip>Parameter type</ToolTip>
+          <Default>Param6Type</Default>
+          <Function>
+          </Function>
+            </Literal>
+        </Declarations>
+          <Code Language="cpp"><![CDATA[DECLARE_DYNAMIC_DELEGATE_SixParams($delegateName$, $paramType1$, $paramName1$, $paramType2$, $paramName2$, $paramType3$, $paramName3$, $paramType4$, $paramName4$, $paramType5$, $paramName5$, $paramType6$, $paramName6$)$selected$ $end$]]></Code>
+        </Snippet>
+      </CodeSnippet>
+    </CodeSnippets>
\ No newline at end of file
diff --git a/Engine/Extras/VisualStudioSnippets/Unreal.deldyn7.snippet b/Engine/Extras/VisualStudioSnippets/Unreal.deldyn7.snippet
new file mode 100644
--- /dev/null
+++ b/Engine/Extras/VisualStudioSnippets/Unreal.deldyn7.snippet
@@ -0,0 +1,126 @@
+<?xml version="1.0" encoding="utf-8"?>
+    <CodeSnippets xmlns="http://schemas.microsoft.com/VisualStudio/2005/CodeSnippet">
+      <CodeSnippet Format="1.0.0">
+        <Header>
+          <SnippetTypes>
+            <SnippetType>Expansion</SnippetType>
+          </SnippetTypes>
+          <Title>unreal-deldyn7</Title>
+          <Author>Epic Games</Author>
+          <Description>Unreal Engine Dynamic Delegate snippet</Description>
+          <HelpUrl>
+          </HelpUrl>
+          <Shortcut>unreal-deldyn7</Shortcut>
+        </Header>
+        <Snippet>
+          <Declarations>
+       <Literal Editable="true">
+          <ID>delegateName</ID>
+          <ToolTip>Delegate type name</ToolTip>
+          <Default>DelegateName</Default>
+          <Function>
+          </Function>
+        </Literal>
+        <Literal Editable="true">
+          <ID>paramName1</ID>
+          <ToolTip>Parameter name</ToolTip>
+          <Default>Param1Name</Default>
+          <Function>
+          </Function>
+            </Literal>
+        <Literal Editable="true">
+          <ID>paramName2</ID>
+          <ToolTip>Parameter name</ToolTip>
+          <Default>Param2Name</Default>
+          <Function>
+          </Function>
+            </Literal>
+        <Literal Editable="true">
+          <ID>paramName3</ID>
+          <ToolTip>Parameter name</ToolTip>
+          <Default>Param3Name</Default>
+          <Function>
+          </Function>
+            </Literal>
+        <Literal Editable="true">
+          <ID>paramName4</ID>
+          <ToolTip>Parameter name</ToolTip>
+          <Default>Param4Name</Default>
+          <Function>
+          </Function>
+            </Literal>
+        <Literal Editable="true">
+          <ID>paramName5</ID>
+          <ToolTip>Parameter name</ToolTip>
+          <Default>Param5Name</Default>
+          <Function>
+          </Function>
+            </Literal>
+        <Literal Editable="true">
+          <ID>paramName6</ID>
+          <ToolTip>Parameter name</ToolTip>
+          <Default>Param6Name</Default>
+          <Function>
+          </Function>
+            </Literal>
+        <Literal Editable="true">
+          <ID>paramName7</ID>
+          <ToolTip>Parameter name</ToolTip>
+          <Default>Param7Name</Default>
+          <Function>
+          </Function>
+            </Literal>
+       <Literal Editable="true">
+          <ID>paramType1</ID>
+          <ToolTip>Parameter type</ToolTip>
+          <Default>Param1Type</Default>
+          <Function>
+          </Function>
+            </Literal>
+       <Literal Editable="true">
+          <ID>paramType2</ID>
+          <ToolTip>Parameter type</ToolTip>
+          <Default>Param2Type</Default>
+          <Function>
+          </Function>
+            </Literal>
+       <Literal Editable="true">
+          <ID>paramType3</ID>
+          <ToolTip>Parameter type</ToolTip>
+          <Default>Param3Type</Default>
+          <Function>
+          </Function>
+            </Literal>
+       <Literal Editable="true">
+          <ID>paramType4</ID>
+          <ToolTip>Parameter type</ToolTip>
+          <Default>Param4Type</Default>
+          <Function>
+          </Function>
+            </Literal>
+       <Literal Editable="true">
+          <ID>paramType5</ID>
+          <ToolTip>Parameter type</ToolTip>
+          <Default>Param5Type</Default>
+          <Function>
+          </Function>
+            </Literal>
+       <Literal Editable="true">
+          <ID>paramType6</ID>
+          <ToolTip>Parameter type</ToolTip>
+          <Default>Param6Type</Default>
+          <Function>
+          </Function>
+            </Literal>
+       <Literal Editable="true">
+          <ID>paramType7</ID>
+          <ToolTip>Parameter type</ToolTip>
+          <Default>Param7Type</Default>
+          <Function>
+          </Function>
+            </Literal>
+        </Declarations>
+          <Code Language="cpp"><![CDATA[DECLARE_DYNAMIC_DELEGATE_SevenParams($delegateName$, $paramType1$, $paramName1$, $paramType2$, $paramName2$, $paramType3$, $paramName3$, $paramType4$, $paramName4$, $paramType5$, $paramName5$, $paramType6$, $paramName6$, $paramType7$, $paramName7$)$selected$ $end$]]></Code>
+        </Snippet>
+      </CodeSnippet>
+    </CodeSnippets>
\ No newline at end of file
diff --git a/Engine/Extras/VisualStudioSnippets/Unreal.deldyn8.snippet b/Engine/Extras/VisualStudioSnippets/Unreal.deldyn8.snippet
new file mode 100644
--- /dev/null
+++ b/Engine/Extras/VisualStudioSnippets/Unreal.deldyn8.snippet
@@ -0,0 +1,140 @@
+<?xml version="1.0" encoding="utf-8"?>
+    <CodeSnippets xmlns="http://schemas.microsoft.com/VisualStudio/2005/CodeSnippet">
+      <CodeSnippet Format="1.0.0">
+        <Header>
+          <SnippetTypes>
+            <SnippetType>Expansion</SnippetType>
+          </SnippetTypes>
+          <Title>unreal-deldyn8</Title>
+          <Author>Epic Games</Author>
+          <Description>Unreal Engine Dynamic Delegate snippet</Description>
+          <HelpUrl>
+          </HelpUrl>
+          <Shortcut>unreal-deldyn8</Shortcut>
+        </Header>
+        <Snippet>
+          <Declarations>
+       <Literal Editable="true">
+          <ID>delegateName</ID>
+          <ToolTip>Delegate type name</ToolTip>
+          <Default>DelegateName</Default>
+          <Function>
+          </Function>
+        </Literal>
+        <Literal Editable="true">
+          <ID>paramName1</ID>
+          <ToolTip>Parameter name</ToolTip>
+          <Default>Param1Name</Default>
+          <Function>
+          </Function>
+            </Literal>
+        <Literal Editable="true">
+          <ID>paramName2</ID>
+          <ToolTip>Parameter name</ToolTip>
+          <Default>Param2Name</Default>
+          <Function>
+          </Function>
+            </Literal>
+        <Literal Editable="true">
+          <ID>paramName3</ID>
+          <ToolTip>Parameter name</ToolTip>
+          <Default>Param3Name</Default>
+          <Function>
+          </Function>
+            </Literal>
+        <Literal Editable="true">
+          <ID>paramName4</ID>
+          <ToolTip>Parameter name</ToolTip>
+          <Default>Param4Name</Default>
+          <Function>
+          </Function>
+            </Literal>
+        <Literal Editable="true">
+          <ID>paramName5</ID>
+          <ToolTip>Parameter name</ToolTip>
+          <Default>Param5Name</Default>
+          <Function>
+          </Function>
+            </Literal>
+        <Literal Editable="true">
+          <ID>paramName6</ID>
+          <ToolTip>Parameter name</ToolTip>
+          <Default>Param6Name</Default>
+          <Function>
+          </Function>
+            </Literal>
+        <Literal Editable="true">
+          <ID>paramName7</ID>
+          <ToolTip>Parameter name</ToolTip>
+          <Default>Param7Name</Default>
+          <Function>
+          </Function>
+            </Literal>
+        <Literal Editable="true">
+          <ID>paramName8</ID>
+          <ToolTip>Parameter name</ToolTip>
+          <Default>Param8Name</Default>
+          <Function>
+          </Function>
+            </Literal>
+       <Literal Editable="true">
+          <ID>paramType1</ID>
+          <ToolTip>Parameter type</ToolTip>
+          <Default>Param1Type</Default>
+          <Function>
+          </Function>
+            </Literal>
+       <Literal Editable="true">
+          <ID>paramType2</ID>
+          <ToolTip>Parameter type</ToolTip>
+          <Default>Param2Type</Default>
+          <Function>
+          </Function>
+            </Literal>
+       <Literal Editable="true">
+          <ID>paramType3</ID>
+          <ToolTip>Parameter type</ToolTip>
+          <Default>Param3Type</Default>
+          <Function>
+          </Function>
+            </Literal>
+       <Literal Editable="true">
+          <ID>paramType4</ID>
+          <ToolTip>Parameter type</ToolTip>
+          <Default>Param4Type</Default>
+          <Function>
+          </Function>
+            </Literal>
+       <Literal Editable="true">
+          <ID>paramType5</ID>
+          <ToolTip>Parameter type</ToolTip>
+          <Default>Param5Type</Default>
+          <Function>
+          </Function>
+            </Literal>
+       <Literal Editable="true">
+          <ID>paramType6</ID>
+          <ToolTip>Parameter type</ToolTip>
+          <Default>Param6Type</Default>
+          <Function>
+          </Function>
+            </Literal>
+       <Literal Editable="true">
+          <ID>paramType7</ID>
+          <ToolTip>Parameter type</ToolTip>
+          <Default>Param7Type</Default>
+          <Function>
+          </Function>
+            </Literal>
+       <Literal Editable="true">
+          <ID>paramType8</ID>
+          <ToolTip>Parameter type</ToolTip>
+          <Default>Param8Type</Default>
+          <Function>
+          </Function>
+            </Literal>
+        </Declarations>
+          <Code Language="cpp"><![CDATA[DECLARE_DYNAMIC_DELEGATE_EightParams($delegateName$, $paramType1$, $paramName1$, $paramType2$, $paramName2$, $paramType3$, $paramName3$, $paramType4$, $paramName4$, $paramType5$, $paramName5$, $paramType6$, $paramName6$, $paramType7$, $paramName7$, $paramType8$, $paramName8$)$selected$ $end$]]></Code>
+        </Snippet>
+      </CodeSnippet>
+    </CodeSnippets>
\ No newline at end of file
diff --git a/Engine/Extras/VisualStudioSnippets/Unreal.deldynmul.snippet b/Engine/Extras/VisualStudioSnippets/Unreal.deldynmul.snippet
new file mode 100644
--- /dev/null
+++ b/Engine/Extras/VisualStudioSnippets/Unreal.deldynmul.snippet
@@ -0,0 +1,28 @@
+<?xml version="1.0" encoding="utf-8"?>
+    <CodeSnippets xmlns="http://schemas.microsoft.com/VisualStudio/2005/CodeSnippet">
+      <CodeSnippet Format="1.0.0">
+        <Header>
+          <SnippetTypes>
+            <SnippetType>Expansion</SnippetType>
+          </SnippetTypes>
+          <Title>unreal-deldynmul</Title>
+          <Author>Epic Games</Author>
+          <Description>Unreal Engine Dynamic Multicast Delegate snippet</Description>
+          <HelpUrl>
+          </HelpUrl>
+          <Shortcut>unreal-deldynmul</Shortcut>
+        </Header>
+        <Snippet>
+          <Declarations>
+       <Literal Editable="true">
+          <ID>delegateName</ID>
+          <ToolTip>Delegate type name</ToolTip>
+          <Default>DelegateName</Default>
+          <Function>
+          </Function>
+        </Literal>
+        </Declarations>
+          <Code Language="cpp"><![CDATA[DECLARE_DYNAMIC_MULTICAST_DELEGATE($delegateName$)$selected$ $end$]]></Code>
+        </Snippet>
+      </CodeSnippet>
+    </CodeSnippets>
\ No newline at end of file
diff --git a/Engine/Extras/VisualStudioSnippets/Unreal.deldynmul1.snippet b/Engine/Extras/VisualStudioSnippets/Unreal.deldynmul1.snippet
new file mode 100644
--- /dev/null
+++ b/Engine/Extras/VisualStudioSnippets/Unreal.deldynmul1.snippet
@@ -0,0 +1,42 @@
+<?xml version="1.0" encoding="utf-8"?>
+    <CodeSnippets xmlns="http://schemas.microsoft.com/VisualStudio/2005/CodeSnippet">
+      <CodeSnippet Format="1.0.0">
+        <Header>
+          <SnippetTypes>
+            <SnippetType>Expansion</SnippetType>
+          </SnippetTypes>
+          <Title>unreal-deldynmul1</Title>
+          <Author>Epic Games</Author>
+          <Description>Unreal Engine Dynamic Multicast Delegate snippet</Description>
+          <HelpUrl>
+          </HelpUrl>
+          <Shortcut>unreal-deldynmul1</Shortcut>
+        </Header>
+        <Snippet>
+          <Declarations>
+       <Literal Editable="true">
+          <ID>delegateName</ID>
+          <ToolTip>Delegate type name</ToolTip>
+          <Default>DelegateName</Default>
+          <Function>
+          </Function>
+        </Literal>
+        <Literal Editable="true">
+          <ID>paramName1</ID>
+          <ToolTip>Parameter name</ToolTip>
+          <Default>Param1Name</Default>
+          <Function>
+          </Function>
+            </Literal>
+       <Literal Editable="true">
+          <ID>paramType1</ID>
+          <ToolTip>Parameter type</ToolTip>
+          <Default>Param1Type</Default>
+          <Function>
+          </Function>
+            </Literal>
+        </Declarations>
+          <Code Language="cpp"><![CDATA[DECLARE_DYNAMIC_MULTICAST_DELEGATE_OneParam($delegateName$, $paramType1$, $paramName1$)$selected$ $end$]]></Code>
+        </Snippet>
+      </CodeSnippet>
+    </CodeSnippets>
\ No newline at end of file
diff --git a/Engine/Extras/VisualStudioSnippets/Unreal.deldynmul2.snippet b/Engine/Extras/VisualStudioSnippets/Unreal.deldynmul2.snippet
new file mode 100644
--- /dev/null
+++ b/Engine/Extras/VisualStudioSnippets/Unreal.deldynmul2.snippet
@@ -0,0 +1,56 @@
+<?xml version="1.0" encoding="utf-8"?>
+    <CodeSnippets xmlns="http://schemas.microsoft.com/VisualStudio/2005/CodeSnippet">
+      <CodeSnippet Format="1.0.0">
+        <Header>
+          <SnippetTypes>
+            <SnippetType>Expansion</SnippetType>
+          </SnippetTypes>
+          <Title>unreal-deldynmul2</Title>
+          <Author>Epic Games</Author>
+          <Description>Unreal Engine Dynamic Multicast Delegate snippet</Description>
+          <HelpUrl>
+          </HelpUrl>
+          <Shortcut>unreal-deldynmul2</Shortcut>
+        </Header>
+        <Snippet>
+          <Declarations>
+       <Literal Editable="true">
+          <ID>delegateName</ID>
+          <ToolTip>Delegate type name</ToolTip>
+          <Default>DelegateName</Default>
+          <Function>
+          </Function>
+        </Literal>
+        <Literal Editable="true">
+          <ID>paramName1</ID>
+          <ToolTip>Parameter name</ToolTip>
+          <Default>Param1Name</Default>
+          <Function>
+          </Function>
+            </Literal>
+        <Literal Editable="true">
+          <ID>paramName2</ID>
+          <ToolTip>Parameter name</ToolTip>
+          <Default>Param2Name</Default>
+          <Function>
+          </Function>
+            </Literal>
+       <Literal Editable="true">
+          <ID>paramType1</ID>
+          <ToolTip>Parameter type</ToolTip>
+          <Default>Param1Type</Default>
+          <Function>
+          </Function>
+            </Literal>
+       <Literal Editable="true">
+          <ID>paramType2</ID>
+          <ToolTip>Parameter type</ToolTip>
+          <Default>Param2Type</Default>
+          <Function>
+          </Function>
+            </Literal>
+        </Declarations>
+          <Code Language="cpp"><![CDATA[DECLARE_DYNAMIC_MULTICAST_DELEGATE_TwoParams($delegateName$, $paramType1$, $paramName1$, $paramType2$, $paramName2$)$selected$ $end$]]></Code>
+        </Snippet>
+      </CodeSnippet>
+    </CodeSnippets>
\ No newline at end of file
diff --git a/Engine/Extras/VisualStudioSnippets/Unreal.deldynmul3.snippet b/Engine/Extras/VisualStudioSnippets/Unreal.deldynmul3.snippet
new file mode 100644
--- /dev/null
+++ b/Engine/Extras/VisualStudioSnippets/Unreal.deldynmul3.snippet
@@ -0,0 +1,70 @@
+<?xml version="1.0" encoding="utf-8"?>
+    <CodeSnippets xmlns="http://schemas.microsoft.com/VisualStudio/2005/CodeSnippet">
+      <CodeSnippet Format="1.0.0">
+        <Header>
+          <SnippetTypes>
+            <SnippetType>Expansion</SnippetType>
+          </SnippetTypes>
+          <Title>unreal-deldynmul3</Title>
+          <Author>Epic Games</Author>
+          <Description>Unreal Engine Dynamic Multicast Delegate snippet</Description>
+          <HelpUrl>
+          </HelpUrl>
+          <Shortcut>unreal-deldynmul3</Shortcut>
+        </Header>
+        <Snippet>
+          <Declarations>
+       <Literal Editable="true">
+          <ID>delegateName</ID>
+          <ToolTip>Delegate type name</ToolTip>
+          <Default>DelegateName</Default>
+          <Function>
+          </Function>
+        </Literal>
+        <Literal Editable="true">
+          <ID>paramName1</ID>
+          <ToolTip>Parameter name</ToolTip>
+          <Default>Param1Name</Default>
+          <Function>
+          </Function>
+            </Literal>
+        <Literal Editable="true">
+          <ID>paramName2</ID>
+          <ToolTip>Parameter name</ToolTip>
+          <Default>Param2Name</Default>
+          <Function>
+          </Function>
+            </Literal>
+        <Literal Editable="true">
+          <ID>paramName3</ID>
+          <ToolTip>Parameter name</ToolTip>
+          <Default>Param3Name</Default>
+          <Function>
+          </Function>
+            </Literal>
+       <Literal Editable="true">
+          <ID>paramType1</ID>
+          <ToolTip>Parameter type</ToolTip>
+          <Default>Param1Type</Default>
+          <Function>
+          </Function>
+            </Literal>
+       <Literal Editable="true">
+          <ID>paramType2</ID>
+          <ToolTip>Parameter type</ToolTip>
+          <Default>Param2Type</Default>
+          <Function>
+          </Function>
+            </Literal>
+       <Literal Editable="true">
+          <ID>paramType3</ID>
+          <ToolTip>Parameter type</ToolTip>
+          <Default>Param3Type</Default>
+          <Function>
+          </Function>
+            </Literal>
+        </Declarations>
+          <Code Language="cpp"><![CDATA[DECLARE_DYNAMIC_MULTICAST_DELEGATE_ThreeParams($delegateName$, $paramType1$, $paramName1$, $paramType2$, $paramName2$, $paramType3$, $paramName3$)$selected$ $end$]]></Code>
+        </Snippet>
+      </CodeSnippet>
+    </CodeSnippets>
\ No newline at end of file
diff --git a/Engine/Extras/VisualStudioSnippets/Unreal.deldynmul4.snippet b/Engine/Extras/VisualStudioSnippets/Unreal.deldynmul4.snippet
new file mode 100644
--- /dev/null
+++ b/Engine/Extras/VisualStudioSnippets/Unreal.deldynmul4.snippet
@@ -0,0 +1,84 @@
+<?xml version="1.0" encoding="utf-8"?>
+    <CodeSnippets xmlns="http://schemas.microsoft.com/VisualStudio/2005/CodeSnippet">
+      <CodeSnippet Format="1.0.0">
+        <Header>
+          <SnippetTypes>
+            <SnippetType>Expansion</SnippetType>
+          </SnippetTypes>
+          <Title>unreal-deldynmul4</Title>
+          <Author>Epic Games</Author>
+          <Description>Unreal Engine Dynamic Multicast Delegate snippet</Description>
+          <HelpUrl>
+          </HelpUrl>
+          <Shortcut>unreal-deldynmul4</Shortcut>
+        </Header>
+        <Snippet>
+          <Declarations>
+       <Literal Editable="true">
+          <ID>delegateName</ID>
+          <ToolTip>Delegate type name</ToolTip>
+          <Default>DelegateName</Default>
+          <Function>
+          </Function>
+        </Literal>
+        <Literal Editable="true">
+          <ID>paramName1</ID>
+          <ToolTip>Parameter name</ToolTip>
+          <Default>Param1Name</Default>
+          <Function>
+          </Function>
+            </Literal>
+        <Literal Editable="true">
+          <ID>paramName2</ID>
+          <ToolTip>Parameter name</ToolTip>
+          <Default>Param2Name</Default>
+          <Function>
+          </Function>
+            </Literal>
+        <Literal Editable="true">
+          <ID>paramName3</ID>
+          <ToolTip>Parameter name</ToolTip>
+          <Default>Param3Name</Default>
+          <Function>
+          </Function>
+            </Literal>
+        <Literal Editable="true">
+          <ID>paramName4</ID>
+          <ToolTip>Parameter name</ToolTip>
+          <Default>Param4Name</Default>
+          <Function>
+          </Function>
+            </Literal>
+       <Literal Editable="true">
+          <ID>paramType1</ID>
+          <ToolTip>Parameter type</ToolTip>
+          <Default>Param1Type</Default>
+          <Function>
+          </Function>
+            </Literal>
+       <Literal Editable="true">
+          <ID>paramType2</ID>
+          <ToolTip>Parameter type</ToolTip>
+          <Default>Param2Type</Default>
+          <Function>
+          </Function>
+            </Literal>
+       <Literal Editable="true">
+          <ID>paramType3</ID>
+          <ToolTip>Parameter type</ToolTip>
+          <Default>Param3Type</Default>
+          <Function>
+          </Function>
+            </Literal>
+       <Literal Editable="true">
+          <ID>paramType4</ID>
+          <ToolTip>Parameter type</ToolTip>
+          <Default>Param4Type</Default>
+          <Function>
+          </Function>
+            </Literal>
+        </Declarations>
+          <Code Language="cpp"><![CDATA[DECLARE_DYNAMIC_MULTICAST_DELEGATE_FourParams($delegateName$, $paramType1$, $paramName1$, $paramType2$, $paramName2$, $paramType3$, $paramName3$, $paramType4$, $paramName4$)$selected$ $end$]]></Code>
+        </Snippet>
+      </CodeSnippet>
+    </CodeSnippets>
\ No newline at end of file
diff --git a/Engine/Extras/VisualStudioSnippets/Unreal.deldynmul5.snippet b/Engine/Extras/VisualStudioSnippets/Unreal.deldynmul5.snippet
new file mode 100644
--- /dev/null
+++ b/Engine/Extras/VisualStudioSnippets/Unreal.deldynmul5.snippet
@@ -0,0 +1,98 @@
+<?xml version="1.0" encoding="utf-8"?>
+    <CodeSnippets xmlns="http://schemas.microsoft.com/VisualStudio/2005/CodeSnippet">
+      <CodeSnippet Format="1.0.0">
+        <Header>
+          <SnippetTypes>
+            <SnippetType>Expansion</SnippetType>
+          </SnippetTypes>
+          <Title>unreal-deldynmul5</Title>
+          <Author>Epic Games</Author>
+          <Description>Unreal Engine Dynamic Multicast Delegate snippet</Description>
+          <HelpUrl>
+          </HelpUrl>
+          <Shortcut>unreal-deldynmul5</Shortcut>
+        </Header>
+        <Snippet>
+          <Declarations>
+       <Literal Editable="true">
+          <ID>delegateName</ID>
+          <ToolTip>Delegate type name</ToolTip>
+          <Default>DelegateName</Default>
+          <Function>
+          </Function>
+        </Literal>
+        <Literal Editable="true">
+          <ID>paramName1</ID>
+          <ToolTip>Parameter name</ToolTip>
+          <Default>Param1Name</Default>
+          <Function>
+          </Function>
+            </Literal>
+        <Literal Editable="true">
+          <ID>paramName2</ID>
+          <ToolTip>Parameter name</ToolTip>
+          <Default>Param2Name</Default>
+          <Function>
+          </Function>
+            </Literal>
+        <Literal Editable="true">
+          <ID>paramName3</ID>
+          <ToolTip>Parameter name</ToolTip>
+          <Default>Param3Name</Default>
+          <Function>
+          </Function>
+            </Literal>
+        <Literal Editable="true">
+          <ID>paramName4</ID>
+          <ToolTip>Parameter name</ToolTip>
+          <Default>Param4Name</Default>
+          <Function>
+          </Function>
+            </Literal>
+        <Literal Editable="true">
+          <ID>paramName5</ID>
+          <ToolTip>Parameter name</ToolTip>
+          <Default>Param5Name</Default>
+          <Function>
+          </Function>
+            </Literal>
+       <Literal Editable="true">
+          <ID>paramType1</ID>
+          <ToolTip>Parameter type</ToolTip>
+          <Default>Param1Type</Default>
+          <Function>
+          </Function>
+            </Literal>
+       <Literal Editable="true">
+          <ID>paramType2</ID>
+          <ToolTip>Parameter type</ToolTip>
+          <Default>Param2Type</Default>
+          <Function>
+          </Function>
+            </Literal>
+       <Literal Editable="true">
+          <ID>paramType3</ID>
+          <ToolTip>Parameter type</ToolTip>
+          <Default>Param3Type</Default>
+          <Function>
+          </Function>
+            </Literal>
+       <Literal Editable="true">
+          <ID>paramType4</ID>
+          <ToolTip>Parameter type</ToolTip>
+          <Default>Param4Type</Default>
+          <Function>
+          </Function>
+            </Literal>
+       <Literal Editable="true">
+          <ID>paramType5</ID>
+          <ToolTip>Parameter type</ToolTip>
+          <Default>Param5Type</Default>
+          <Function>
+          </Function>
+            </Literal>
+        </Declarations>
+          <Code Language="cpp"><![CDATA[DECLARE_DYNAMIC_MULTICAST_DELEGATE_FiveParams($delegateName$, $paramType1$, $paramName1$, $paramType2$, $paramName2$, $paramType3$, $paramName3$, $paramType4$, $paramName4$, $paramType5$, $paramName5$)$selected$ $end$]]></Code>
+        </Snippet>
+      </CodeSnippet>
+    </CodeSnippets>
\ No newline at end of file
diff --git a/Engine/Extras/VisualStudioSnippets/Unreal.deldynmul7.snippet b/Engine/Extras/VisualStudioSnippets/Unreal.deldynmul7.snippet
new file mode 100644
--- /dev/null
+++ b/Engine/Extras/VisualStudioSnippets/Unreal.deldynmul7.snippet
@@ -0,0 +1,126 @@
+<?xml version="1.0" encoding="utf-8"?>
+    <CodeSnippets xmlns="http://schemas.microsoft.com/VisualStudio/2005/CodeSnippet">
+      <CodeSnippet Format="1.0.0">
+        <Header>
+          <SnippetTypes>
+            <SnippetType>Expansion</SnippetType>
+          </SnippetTypes>
+          <Title>unreal-deldynmul7</Title>
+          <Author>Epic Games</Author>
+          <Description>Unreal Engine Dynamic Multicast Delegate snippet</Description>
+          <HelpUrl>
+          </HelpUrl>
+          <Shortcut>unreal-deldynmul7</Shortcut>
+        </Header>
+        <Snippet>
+          <Declarations>
+       <Literal Editable="true">
+          <ID>delegateName</ID>
+          <ToolTip>Delegate type name</ToolTip>
+          <Default>DelegateName</Default>
+          <Function>
+          </Function>
+        </Literal>
+        <Literal Editable="true">
+          <ID>paramName1</ID>
+          <ToolTip>Parameter name</ToolTip>
+          <Default>Param1Name</Default>
+          <Function>
+          </Function>
+            </Literal>
+        <Literal Editable="true">
+          <ID>paramName2</ID>
+          <ToolTip>Parameter name</ToolTip>
+          <Default>Param2Name</Default>
+          <Function>
+          </Function>
+            </Literal>
+        <Literal Editable="true">
+          <ID>paramName3</ID>
+          <ToolTip>Parameter name</ToolTip>
+          <Default>Param3Name</Default>
+          <Function>
+          </Function>
+            </Literal>
+        <Literal Editable="true">
+          <ID>paramName4</ID>
+          <ToolTip>Parameter name</ToolTip>
+          <Default>Param4Name</Default>
+          <Function>
+          </Function>
+            </Literal>
+        <Literal Editable="true">
+          <ID>paramName5</ID>
+          <ToolTip>Parameter name</ToolTip>
+          <Default>Param5Name</Default>
+          <Function>
+          </Function>
+            </Literal>
+        <Literal Editable="true">
+          <ID>paramName6</ID>
+          <ToolTip>Parameter name</ToolTip>
+          <Default>Param6Name</Default>
+          <Function>
+          </Function>
+            </Literal>
+        <Literal Editable="true">
+          <ID>paramName7</ID>
+          <ToolTip>Parameter name</ToolTip>
+          <Default>Param7Name</Default>
+          <Function>
+          </Function>
+            </Literal>
+       <Literal Editable="true">
+          <ID>paramType1</ID>
+          <ToolTip>Parameter type</ToolTip>
+          <Default>Param1Type</Default>
+          <Function>
+          </Function>
+            </Literal>
+       <Literal Editable="true">
+          <ID>paramType2</ID>
+          <ToolTip>Parameter type</ToolTip>
+          <Default>Param2Type</Default>
+          <Function>
+          </Function>
+            </Literal>
+       <Literal Editable="true">
+          <ID>paramType3</ID>
+          <ToolTip>Parameter type</ToolTip>
+          <Default>Param3Type</Default>
+          <Function>
+          </Function>
+            </Literal>
+       <Literal Editable="true">
+          <ID>paramType4</ID>
+          <ToolTip>Parameter type</ToolTip>
+          <Default>Param4Type</Default>
+          <Function>
+          </Function>
+            </Literal>
+       <Literal Editable="true">
+          <ID>paramType5</ID>
+          <ToolTip>Parameter type</ToolTip>
+          <Default>Param5Type</Default>
+          <Function>
+          </Function>
+            </Literal>
+       <Literal Editable="true">
+          <ID>paramType6</ID>
+          <ToolTip>Parameter type</ToolTip>
+          <Default>Param6Type</Default>
+          <Function>
+          </Function>
+            </Literal>
+       <Literal Editable="true">
+          <ID>paramType7</ID>
+          <ToolTip>Parameter type</ToolTip>
+          <Default>Param7Type</Default>
+          <Function>
+          </Function>
+            </Literal>
+        </Declarations>
+          <Code Language="cpp"><![CDATA[DECLARE_DYNAMIC_MULTICAST_DELEGATE_SevenParams($delegateName$, $paramType1$, $paramName1$, $paramType2$, $paramName2$, $paramType3$, $paramName3$, $paramType4$, $paramName4$, $paramType5$, $paramName5$, $paramType6$, $paramName6$, $paramType7$, $paramName7$)$selected$ $end$]]></Code>
+        </Snippet>
+      </CodeSnippet>
+    </CodeSnippets>
\ No newline at end of file
diff --git a/Engine/Extras/VisualStudioSnippets/Unreal.deldynmul8.snippet b/Engine/Extras/VisualStudioSnippets/Unreal.deldynmul8.snippet
new file mode 100644
--- /dev/null
+++ b/Engine/Extras/VisualStudioSnippets/Unreal.deldynmul8.snippet
@@ -0,0 +1,140 @@
+<?xml version="1.0" encoding="utf-8"?>
+    <CodeSnippets xmlns="http://schemas.microsoft.com/VisualStudio/2005/CodeSnippet">
+      <CodeSnippet Format="1.0.0">
+        <Header>
+          <SnippetTypes>
+            <SnippetType>Expansion</SnippetType>
+          </SnippetTypes>
+          <Title>unreal-deldynmul8</Title>
+          <Author>Epic Games</Author>
+          <Description>Unreal Engine Dynamic Multicast Delegate snippet</Description>
+          <HelpUrl>
+          </HelpUrl>
+          <Shortcut>unreal-deldynmul8</Shortcut>
+        </Header>
+        <Snippet>
+          <Declarations>
+       <Literal Editable="true">
+          <ID>delegateName</ID>
+          <ToolTip>Delegate type name</ToolTip>
+          <Default>DelegateName</Default>
+          <Function>
+          </Function>
+        </Literal>
+        <Literal Editable="true">
+          <ID>paramName1</ID>
+          <ToolTip>Parameter name</ToolTip>
+          <Default>Param1Name</Default>
+          <Function>
+          </Function>
+            </Literal>
+        <Literal Editable="true">
+          <ID>paramName2</ID>
+          <ToolTip>Parameter name</ToolTip>
+          <Default>Param2Name</Default>
+          <Function>
+          </Function>
+            </Literal>
+        <Literal Editable="true">
+          <ID>paramName3</ID>
+          <ToolTip>Parameter name</ToolTip>
+          <Default>Param3Name</Default>
+          <Function>
+          </Function>
+            </Literal>
+        <Literal Editable="true">
+          <ID>paramName4</ID>
+          <ToolTip>Parameter name</ToolTip>
+          <Default>Param4Name</Default>
+          <Function>
+          </Function>
+            </Literal>
+        <Literal Editable="true">
+          <ID>paramName5</ID>
+          <ToolTip>Parameter name</ToolTip>
+          <Default>Param5Name</Default>
+          <Function>
+          </Function>
+            </Literal>
+        <Literal Editable="true">
+          <ID>paramName6</ID>
+          <ToolTip>Parameter name</ToolTip>
+          <Default>Param6Name</Default>
+          <Function>
+          </Function>
+            </Literal>
+        <Literal Editable="true">
+          <ID>paramName7</ID>
+          <ToolTip>Parameter name</ToolTip>
+          <Default>Param7Name</Default>
+          <Function>
+          </Function>
+            </Literal>
+        <Literal Editable="true">
+          <ID>paramName8</ID>
+          <ToolTip>Parameter name</ToolTip>
+          <Default>Param8Name</Default>
+          <Function>
+          </Function>
+            </Literal>
+       <Literal Editable="true">
+          <ID>paramType1</ID>
+          <ToolTip>Parameter type</ToolTip>
+          <Default>Param1Type</Default>
+          <Function>
+          </Function>
+            </Literal>
+       <Literal Editable="true">
+          <ID>paramType2</ID>
+          <ToolTip>Parameter type</ToolTip>
+          <Default>Param2Type</Default>
+          <Function>
+          </Function>
+            </Literal>
+       <Literal Editable="true">
+          <ID>paramType3</ID>
+          <ToolTip>Parameter type</ToolTip>
+          <Default>Param3Type</Default>
+          <Function>
+          </Function>
+            </Literal>
+       <Literal Editable="true">
+          <ID>paramType4</ID>
+          <ToolTip>Parameter type</ToolTip>
+          <Default>Param4Type</Default>
+          <Function>
+          </Function>
+            </Literal>
+       <Literal Editable="true">
+          <ID>paramType5</ID>
+          <ToolTip>Parameter type</ToolTip>
+          <Default>Param5Type</Default>
+          <Function>
+          </Function>
+            </Literal>
+       <Literal Editable="true">
+          <ID>paramType6</ID>
+          <ToolTip>Parameter type</ToolTip>
+          <Default>Param6Type</Default>
+          <Function>
+          </Function>
+            </Literal>
+       <Literal Editable="true">
+          <ID>paramType7</ID>
+          <ToolTip>Parameter type</ToolTip>
+          <Default>Param7Type</Default>
+          <Function>
+          </Function>
+            </Literal>
+       <Literal Editable="true">
+          <ID>paramType8</ID>
+          <ToolTip>Parameter type</ToolTip>
+          <Default>Param8Type</Default>
+          <Function>
+          </Function>
+            </Literal>
+        </Declarations>
+          <Code Language="cpp"><![CDATA[DECLARE_DYNAMIC_MULTICAST_DELEGATE_EightParams($delegateName$, $paramType1$, $paramName1$, $paramType2$, $paramName2$, $paramType3$, $paramName3$, $paramType4$, $paramName4$, $paramType5$, $paramName5$, $paramType6$, $paramName6$, $paramType7$, $paramName7$, $paramType8$, $paramName8$)$selected$ $end$]]></Code>
+        </Snippet>
+      </CodeSnippet>
+    </CodeSnippets>
\ No newline at end of file
diff --git a/Engine/Extras/VisualStudioSnippets/Unreal.delmul.snippet b/Engine/Extras/VisualStudioSnippets/Unreal.delmul.snippet
new file mode 100644
--- /dev/null
+++ b/Engine/Extras/VisualStudioSnippets/Unreal.delmul.snippet
@@ -0,0 +1,28 @@
+<?xml version="1.0" encoding="utf-8"?>
+    <CodeSnippets xmlns="http://schemas.microsoft.com/VisualStudio/2005/CodeSnippet">
+      <CodeSnippet Format="1.0.0">
+        <Header>
+          <SnippetTypes>
+            <SnippetType>Expansion</SnippetType>
+          </SnippetTypes>
+          <Title>unreal-delmul</Title>
+          <Author>Epic Games</Author>
+          <Description>Unreal Engine Multicast Delegate snippet</Description>
+          <HelpUrl>
+          </HelpUrl>
+          <Shortcut>unreal-delmul</Shortcut>
+        </Header>
+        <Snippet>
+          <Declarations>
+       <Literal Editable="true">
+          <ID>delegateName</ID>
+          <ToolTip>Delegate type name</ToolTip>
+          <Default>DelegateName</Default>
+          <Function>
+          </Function>
+        </Literal>
+        </Declarations>
+          <Code Language="cpp"><![CDATA[DECLARE_MULTICAST_DELEGATE($delegateName$)$selected$ $end$]]></Code>
+        </Snippet>
+      </CodeSnippet>
+    </CodeSnippets>
\ No newline at end of file
diff --git a/Engine/Extras/VisualStudioSnippets/Unreal.delmul1.snippet b/Engine/Extras/VisualStudioSnippets/Unreal.delmul1.snippet
new file mode 100644
--- /dev/null
+++ b/Engine/Extras/VisualStudioSnippets/Unreal.delmul1.snippet
@@ -0,0 +1,42 @@
+<?xml version="1.0" encoding="utf-8"?>
+    <CodeSnippets xmlns="http://schemas.microsoft.com/VisualStudio/2005/CodeSnippet">
+      <CodeSnippet Format="1.0.0">
+        <Header>
+          <SnippetTypes>
+            <SnippetType>Expansion</SnippetType>
+          </SnippetTypes>
+          <Title>unreal-delmul1</Title>
+          <Author>Epic Games</Author>
+          <Description>Unreal Engine Multicast Delegate snippet</Description>
+          <HelpUrl>
+          </HelpUrl>
+          <Shortcut>unreal-delmul1</Shortcut>
+        </Header>
+        <Snippet>
+          <Declarations>
+       <Literal Editable="true">
+          <ID>delegateName</ID>
+          <ToolTip>Delegate type name</ToolTip>
+          <Default>DelegateName</Default>
+          <Function>
+          </Function>
+        </Literal>
+        <Literal Editable="true">
+          <ID>paramName1</ID>
+          <ToolTip>Parameter name</ToolTip>
+          <Default>Param1Name</Default>
+          <Function>
+          </Function>
+            </Literal>
+       <Literal Editable="true">
+          <ID>paramType1</ID>
+          <ToolTip>Parameter type</ToolTip>
+          <Default>Param1Type</Default>
+          <Function>
+          </Function>
+            </Literal>
+        </Declarations>
+          <Code Language="cpp"><![CDATA[DECLARE_MULTICAST_DELEGATE_OneParam($delegateName$, $paramType1$, $paramName1$)$selected$ $end$]]></Code>
+        </Snippet>
+      </CodeSnippet>
+    </CodeSnippets>
\ No newline at end of file
diff --git a/Engine/Extras/VisualStudioSnippets/Unreal.delmul2.snippet b/Engine/Extras/VisualStudioSnippets/Unreal.delmul2.snippet
new file mode 100644
--- /dev/null
+++ b/Engine/Extras/VisualStudioSnippets/Unreal.delmul2.snippet
@@ -0,0 +1,56 @@
+<?xml version="1.0" encoding="utf-8"?>
+    <CodeSnippets xmlns="http://schemas.microsoft.com/VisualStudio/2005/CodeSnippet">
+      <CodeSnippet Format="1.0.0">
+        <Header>
+          <SnippetTypes>
+            <SnippetType>Expansion</SnippetType>
+          </SnippetTypes>
+          <Title>unreal-delmul2</Title>
+          <Author>Epic Games</Author>
+          <Description>Unreal Engine Multicast Delegate snippet</Description>
+          <HelpUrl>
+          </HelpUrl>
+          <Shortcut>unreal-delmul2</Shortcut>
+        </Header>
+        <Snippet>
+          <Declarations>
+       <Literal Editable="true">
+          <ID>delegateName</ID>
+          <ToolTip>Delegate type name</ToolTip>
+          <Default>DelegateName</Default>
+          <Function>
+          </Function>
+        </Literal>
+        <Literal Editable="true">
+          <ID>paramName1</ID>
+          <ToolTip>Parameter name</ToolTip>
+          <Default>Param1Name</Default>
+          <Function>
+          </Function>
+            </Literal>
+        <Literal Editable="true">
+          <ID>paramName2</ID>
+          <ToolTip>Parameter name</ToolTip>
+          <Default>Param2Name</Default>
+          <Function>
+          </Function>
+            </Literal>
+       <Literal Editable="true">
+          <ID>paramType1</ID>
+          <ToolTip>Parameter type</ToolTip>
+          <Default>Param1Type</Default>
+          <Function>
+          </Function>
+            </Literal>
+       <Literal Editable="true">
+          <ID>paramType2</ID>
+          <ToolTip>Parameter type</ToolTip>
+          <Default>Param2Type</Default>
+          <Function>
+          </Function>
+            </Literal>
+        </Declarations>
+          <Code Language="cpp"><![CDATA[DECLARE_MULTICAST_DELEGATE_TwoParams($delegateName$, $paramType1$, $paramName1$, $paramType2$, $paramName2$)$selected$ $end$]]></Code>
+        </Snippet>
+      </CodeSnippet>
+    </CodeSnippets>
\ No newline at end of file
diff --git a/Engine/Extras/VisualStudioSnippets/Unreal.delmul3.snippet b/Engine/Extras/VisualStudioSnippets/Unreal.delmul3.snippet
new file mode 100644
--- /dev/null
+++ b/Engine/Extras/VisualStudioSnippets/Unreal.delmul3.snippet
@@ -0,0 +1,70 @@
+<?xml version="1.0" encoding="utf-8"?>
+    <CodeSnippets xmlns="http://schemas.microsoft.com/VisualStudio/2005/CodeSnippet">
+      <CodeSnippet Format="1.0.0">
+        <Header>
+          <SnippetTypes>
+            <SnippetType>Expansion</SnippetType>
+          </SnippetTypes>
+          <Title>unreal-delmul3</Title>
+          <Author>Epic Games</Author>
+          <Description>Unreal Engine Multicast Delegate snippet</Description>
+          <HelpUrl>
+          </HelpUrl>
+          <Shortcut>unreal-delmul3</Shortcut>
+        </Header>
+        <Snippet>
+          <Declarations>
+       <Literal Editable="true">
+          <ID>delegateName</ID>
+          <ToolTip>Delegate type name</ToolTip>
+          <Default>DelegateName</Default>
+          <Function>
+          </Function>
+        </Literal>
+        <Literal Editable="true">
+          <ID>paramName1</ID>
+          <ToolTip>Parameter name</ToolTip>
+          <Default>Param1Name</Default>
+          <Function>
+          </Function>
+            </Literal>
+        <Literal Editable="true">
+          <ID>paramName2</ID>
+          <ToolTip>Parameter name</ToolTip>
+          <Default>Param2Name</Default>
+          <Function>
+          </Function>
+            </Literal>
+        <Literal Editable="true">
+          <ID>paramName3</ID>
+          <ToolTip>Parameter name</ToolTip>
+          <Default>Param3Name</Default>
+          <Function>
+          </Function>
+            </Literal>
+       <Literal Editable="true">
+          <ID>paramType1</ID>
+          <ToolTip>Parameter type</ToolTip>
+          <Default>Param1Type</Default>
+          <Function>
+          </Function>
+            </Literal>
+       <Literal Editable="true">
+          <ID>paramType2</ID>
+          <ToolTip>Parameter type</ToolTip>
+          <Default>Param2Type</Default>
+          <Function>
+          </Function>
+            </Literal>
+       <Literal Editable="true">
+          <ID>paramType3</ID>
+          <ToolTip>Parameter type</ToolTip>
+          <Default>Param3Type</Default>
+          <Function>
+          </Function>
+            </Literal>
+        </Declarations>
+          <Code Language="cpp"><![CDATA[DECLARE_MULTICAST_DELEGATE_ThreeParams($delegateName$, $paramType1$, $paramName1$, $paramType2$, $paramName2$, $paramType3$, $paramName3$)$selected$ $end$]]></Code>
+        </Snippet>
+      </CodeSnippet>
+    </CodeSnippets>
\ No newline at end of file
diff --git a/Engine/Extras/VisualStudioSnippets/Unreal.delmul4.snippet b/Engine/Extras/VisualStudioSnippets/Unreal.delmul4.snippet
new file mode 100644
--- /dev/null
+++ b/Engine/Extras/VisualStudioSnippets/Unreal.delmul4.snippet
@@ -0,0 +1,84 @@
+<?xml version="1.0" encoding="utf-8"?>
+    <CodeSnippets xmlns="http://schemas.microsoft.com/VisualStudio/2005/CodeSnippet">
+      <CodeSnippet Format="1.0.0">
+        <Header>
+          <SnippetTypes>
+            <SnippetType>Expansion</SnippetType>
+          </SnippetTypes>
+          <Title>unreal-delmul4</Title>
+          <Author>Epic Games</Author>
+          <Description>Unreal Engine Multicast Delegate snippet</Description>
+          <HelpUrl>
+          </HelpUrl>
+          <Shortcut>unreal-delmul4</Shortcut>
+        </Header>
+        <Snippet>
+          <Declarations>
+       <Literal Editable="true">
+          <ID>delegateName</ID>
+          <ToolTip>Delegate type name</ToolTip>
+          <Default>DelegateName</Default>
+          <Function>
+          </Function>
+        </Literal>
+        <Literal Editable="true">
+          <ID>paramName1</ID>
+          <ToolTip>Parameter name</ToolTip>
+          <Default>Param1Name</Default>
+          <Function>
+          </Function>
+            </Literal>
+        <Literal Editable="true">
+          <ID>paramName2</ID>
+          <ToolTip>Parameter name</ToolTip>
+          <Default>Param2Name</Default>
+          <Function>
+          </Function>
+            </Literal>
+        <Literal Editable="true">
+          <ID>paramName3</ID>
+          <ToolTip>Parameter name</ToolTip>
+          <Default>Param3Name</Default>
+          <Function>
+          </Function>
+            </Literal>
+        <Literal Editable="true">
+          <ID>paramName4</ID>
+          <ToolTip>Parameter name</ToolTip>
+          <Default>Param4Name</Default>
+          <Function>
+          </Function>
+            </Literal>
+       <Literal Editable="true">
+          <ID>paramType1</ID>
+          <ToolTip>Parameter type</ToolTip>
+          <Default>Param1Type</Default>
+          <Function>
+          </Function>
+            </Literal>
+       <Literal Editable="true">
+          <ID>paramType2</ID>
+          <ToolTip>Parameter type</ToolTip>
+          <Default>Param2Type</Default>
+          <Function>
+          </Function>
+            </Literal>
+       <Literal Editable="true">
+          <ID>paramType3</ID>
+          <ToolTip>Parameter type</ToolTip>
+          <Default>Param3Type</Default>
+          <Function>
+          </Function>
+            </Literal>
+       <Literal Editable="true">
+          <ID>paramType4</ID>
+          <ToolTip>Parameter type</ToolTip>
+          <Default>Param4Type</Default>
+          <Function>
+          </Function>
+            </Literal>
+        </Declarations>
+          <Code Language="cpp"><![CDATA[DECLARE_MULTICAST_DELEGATE_FourParams($delegateName$, $paramType1$, $paramName1$, $paramType2$, $paramName2$, $paramType3$, $paramName3$, $paramType4$, $paramName4$)$selected$ $end$]]></Code>
+        </Snippet>
+      </CodeSnippet>
+    </CodeSnippets>
\ No newline at end of file
diff --git a/Engine/Extras/VisualStudioSnippets/Unreal.delmul5.snippet b/Engine/Extras/VisualStudioSnippets/Unreal.delmul5.snippet
new file mode 100644
--- /dev/null
+++ b/Engine/Extras/VisualStudioSnippets/Unreal.delmul5.snippet
@@ -0,0 +1,98 @@
+<?xml version="1.0" encoding="utf-8"?>
+    <CodeSnippets xmlns="http://schemas.microsoft.com/VisualStudio/2005/CodeSnippet">
+      <CodeSnippet Format="1.0.0">
+        <Header>
+          <SnippetTypes>
+            <SnippetType>Expansion</SnippetType>
+          </SnippetTypes>
+          <Title>unreal-delmul5</Title>
+          <Author>Epic Games</Author>
+          <Description>Unreal Engine Multicast Delegate snippet</Description>
+          <HelpUrl>
+          </HelpUrl>
+          <Shortcut>unreal-delmul5</Shortcut>
+        </Header>
+        <Snippet>
+          <Declarations>
+       <Literal Editable="true">
+          <ID>delegateName</ID>
+          <ToolTip>Delegate type name</ToolTip>
+          <Default>DelegateName</Default>
+          <Function>
+          </Function>
+        </Literal>
+        <Literal Editable="true">
+          <ID>paramName1</ID>
+          <ToolTip>Parameter name</ToolTip>
+          <Default>Param1Name</Default>
+          <Function>
+          </Function>
+            </Literal>
+        <Literal Editable="true">
+          <ID>paramName2</ID>
+          <ToolTip>Parameter name</ToolTip>
+          <Default>Param2Name</Default>
+          <Function>
+          </Function>
+            </Literal>
+        <Literal Editable="true">
+          <ID>paramName3</ID>
+          <ToolTip>Parameter name</ToolTip>
+          <Default>Param3Name</Default>
+          <Function>
+          </Function>
+            </Literal>
+        <Literal Editable="true">
+          <ID>paramName4</ID>
+          <ToolTip>Parameter name</ToolTip>
+          <Default>Param4Name</Default>
+          <Function>
+          </Function>
+            </Literal>
+        <Literal Editable="true">
+          <ID>paramName5</ID>
+          <ToolTip>Parameter name</ToolTip>
+          <Default>Param5Name</Default>
+          <Function>
+          </Function>
+            </Literal>
+       <Literal Editable="true">
+          <ID>paramType1</ID>
+          <ToolTip>Parameter type</ToolTip>
+          <Default>Param1Type</Default>
+          <Function>
+          </Function>
+            </Literal>
+       <Literal Editable="true">
+          <ID>paramType2</ID>
+          <ToolTip>Parameter type</ToolTip>
+          <Default>Param2Type</Default>
+          <Function>
+          </Function>
+            </Literal>
+       <Literal Editable="true">
+          <ID>paramType3</ID>
+          <ToolTip>Parameter type</ToolTip>
+          <Default>Param3Type</Default>
+          <Function>
+          </Function>
+            </Literal>
+       <Literal Editable="true">
+          <ID>paramType4</ID>
+          <ToolTip>Parameter type</ToolTip>
+          <Default>Param4Type</Default>
+          <Function>
+          </Function>
+            </Literal>
+       <Literal Editable="true">
+          <ID>paramType5</ID>
+          <ToolTip>Parameter type</ToolTip>
+          <Default>Param5Type</Default>
+          <Function>
+          </Function>
+            </Literal>
+        </Declarations>
+          <Code Language="cpp"><![CDATA[DECLARE_MULTICAST_DELEGATE_FiveParams($delegateName$, $paramType1$, $paramName1$, $paramType2$, $paramName2$, $paramType3$, $paramName3$, $paramType4$, $paramName4$, $paramType5$, $paramName5$)$selected$ $end$]]></Code>
+        </Snippet>
+      </CodeSnippet>
+    </CodeSnippets>
\ No newline at end of file
diff --git a/Engine/Extras/VisualStudioSnippets/Unreal.delmul6.snippet b/Engine/Extras/VisualStudioSnippets/Unreal.delmul6.snippet
new file mode 100644
--- /dev/null
+++ b/Engine/Extras/VisualStudioSnippets/Unreal.delmul6.snippet
@@ -0,0 +1,112 @@
+<?xml version="1.0" encoding="utf-8"?>
+    <CodeSnippets xmlns="http://schemas.microsoft.com/VisualStudio/2005/CodeSnippet">
+      <CodeSnippet Format="1.0.0">
+        <Header>
+          <SnippetTypes>
+            <SnippetType>Expansion</SnippetType>
+          </SnippetTypes>
+          <Title>unreal-delmul6</Title>
+          <Author>Epic Games</Author>
+          <Description>Unreal Engine Multicast Delegate snippet</Description>
+          <HelpUrl>
+          </HelpUrl>
+          <Shortcut>unreal-delmul6</Shortcut>
+        </Header>
+        <Snippet>
+          <Declarations>
+       <Literal Editable="true">
+          <ID>delegateName</ID>
+          <ToolTip>Delegate type name</ToolTip>
+          <Default>DelegateName</Default>
+          <Function>
+          </Function>
+        </Literal>
+        <Literal Editable="true">
+          <ID>paramName1</ID>
+          <ToolTip>Parameter name</ToolTip>
+          <Default>Param1Name</Default>
+          <Function>
+          </Function>
+            </Literal>
+        <Literal Editable="true">
+          <ID>paramName2</ID>
+          <ToolTip>Parameter name</ToolTip>
+          <Default>Param2Name</Default>
+          <Function>
+          </Function>
+            </Literal>
+        <Literal Editable="true">
+          <ID>paramName3</ID>
+          <ToolTip>Parameter name</ToolTip>
+          <Default>Param3Name</Default>
+          <Function>
+          </Function>
+            </Literal>
+        <Literal Editable="true">
+          <ID>paramName4</ID>
+          <ToolTip>Parameter name</ToolTip>
+          <Default>Param4Name</Default>
+          <Function>
+          </Function>
+            </Literal>
+        <Literal Editable="true">
+          <ID>paramName5</ID>
+          <ToolTip>Parameter name</ToolTip>
+          <Default>Param5Name</Default>
+          <Function>
+          </Function>
+            </Literal>
+        <Literal Editable="true">
+          <ID>paramName6</ID>
+          <ToolTip>Parameter name</ToolTip>
+          <Default>Param6Name</Default>
+          <Function>
+          </Function>
+            </Literal>
+       <Literal Editable="true">
+          <ID>paramType1</ID>
+          <ToolTip>Parameter type</ToolTip>
+          <Default>Param1Type</Default>
+          <Function>
+          </Function>
+            </Literal>
+       <Literal Editable="true">
+          <ID>paramType2</ID>
+          <ToolTip>Parameter type</ToolTip>
+          <Default>Param2Type</Default>
+          <Function>
+          </Function>
+            </Literal>
+       <Literal Editable="true">
+          <ID>paramType3</ID>
+          <ToolTip>Parameter type</ToolTip>
+          <Default>Param3Type</Default>
+          <Function>
+          </Function>
+            </Literal>
+       <Literal Editable="true">
+          <ID>paramType4</ID>
+          <ToolTip>Parameter type</ToolTip>
+          <Default>Param4Type</Default>
+          <Function>
+          </Function>
+            </Literal>
+       <Literal Editable="true">
+          <ID>paramType5</ID>
+          <ToolTip>Parameter type</ToolTip>
+          <Default>Param5Type</Default>
+          <Function>
+          </Function>
+            </Literal>
+       <Literal Editable="true">
+          <ID>paramType6</ID>
+          <ToolTip>Parameter type</ToolTip>
+          <Default>Param6Type</Default>
+          <Function>
+          </Function>
+            </Literal>
+        </Declarations>
+          <Code Language="cpp"><![CDATA[DECLARE_MULTICAST_DELEGATE_SixParams($delegateName$, $paramType1$, $paramName1$, $paramType2$, $paramName2$, $paramType3$, $paramName3$, $paramType4$, $paramName4$, $paramType5$, $paramName5$, $paramType6$, $paramName6$)$selected$ $end$]]></Code>
+        </Snippet>
+      </CodeSnippet>
+    </CodeSnippets>
\ No newline at end of file
diff --git a/Engine/Extras/VisualStudioSnippets/Unreal.delmul7.snippet b/Engine/Extras/VisualStudioSnippets/Unreal.delmul7.snippet
new file mode 100644
--- /dev/null
+++ b/Engine/Extras/VisualStudioSnippets/Unreal.delmul7.snippet
@@ -0,0 +1,126 @@
+<?xml version="1.0" encoding="utf-8"?>
+    <CodeSnippets xmlns="http://schemas.microsoft.com/VisualStudio/2005/CodeSnippet">
+      <CodeSnippet Format="1.0.0">
+        <Header>
+          <SnippetTypes>
+            <SnippetType>Expansion</SnippetType>
+          </SnippetTypes>
+          <Title>unreal-delmul7</Title>
+          <Author>Epic Games</Author>
+          <Description>Unreal Engine Multicast Delegate snippet</Description>
+          <HelpUrl>
+          </HelpUrl>
+          <Shortcut>unreal-delmul7</Shortcut>
+        </Header>
+        <Snippet>
+          <Declarations>
+       <Literal Editable="true">
+          <ID>delegateName</ID>
+          <ToolTip>Delegate type name</ToolTip>
+          <Default>DelegateName</Default>
+          <Function>
+          </Function>
+        </Literal>
+        <Literal Editable="true">
+          <ID>paramName1</ID>
+          <ToolTip>Parameter name</ToolTip>
+          <Default>Param1Name</Default>
+          <Function>
+          </Function>
+            </Literal>
+        <Literal Editable="true">
+          <ID>paramName2</ID>
+          <ToolTip>Parameter name</ToolTip>
+          <Default>Param2Name</Default>
+          <Function>
+          </Function>
+            </Literal>
+        <Literal Editable="true">
+          <ID>paramName3</ID>
+          <ToolTip>Parameter name</ToolTip>
+          <Default>Param3Name</Default>
+          <Function>
+          </Function>
+            </Literal>
+        <Literal Editable="true">
+          <ID>paramName4</ID>
+          <ToolTip>Parameter name</ToolTip>
+          <Default>Param4Name</Default>
+          <Function>
+          </Function>
+            </Literal>
+        <Literal Editable="true">
+          <ID>paramName5</ID>
+          <ToolTip>Parameter name</ToolTip>
+          <Default>Param5Name</Default>
+          <Function>
+          </Function>
+            </Literal>
+        <Literal Editable="true">
+          <ID>paramName6</ID>
+          <ToolTip>Parameter name</ToolTip>
+          <Default>Param6Name</Default>
+          <Function>
+          </Function>
+            </Literal>
+        <Literal Editable="true">
+          <ID>paramName7</ID>
+          <ToolTip>Parameter name</ToolTip>
+          <Default>Param7Name</Default>
+          <Function>
+          </Function>
+            </Literal>
+       <Literal Editable="true">
+          <ID>paramType1</ID>
+          <ToolTip>Parameter type</ToolTip>
+          <Default>Param1Type</Default>
+          <Function>
+          </Function>
+            </Literal>
+       <Literal Editable="true">
+          <ID>paramType2</ID>
+          <ToolTip>Parameter type</ToolTip>
+          <Default>Param2Type</Default>
+          <Function>
+          </Function>
+            </Literal>
+       <Literal Editable="true">
+          <ID>paramType3</ID>
+          <ToolTip>Parameter type</ToolTip>
+          <Default>Param3Type</Default>
+          <Function>
+          </Function>
+            </Literal>
+       <Literal Editable="true">
+          <ID>paramType4</ID>
+          <ToolTip>Parameter type</ToolTip>
+          <Default>Param4Type</Default>
+          <Function>
+          </Function>
+            </Literal>
+       <Literal Editable="true">
+          <ID>paramType5</ID>
+          <ToolTip>Parameter type</ToolTip>
+          <Default>Param5Type</Default>
+          <Function>
+          </Function>
+            </Literal>
+       <Literal Editable="true">
+          <ID>paramType6</ID>
+          <ToolTip>Parameter type</ToolTip>
+          <Default>Param6Type</Default>
+          <Function>
+          </Function>
+            </Literal>
+       <Literal Editable="true">
+          <ID>paramType7</ID>
+          <ToolTip>Parameter type</ToolTip>
+          <Default>Param7Type</Default>
+          <Function>
+          </Function>
+            </Literal>
+        </Declarations>
+          <Code Language="cpp"><![CDATA[DECLARE_MULTICAST_DELEGATE_SevenParams($delegateName$, $paramType1$, $paramName1$, $paramType2$, $paramName2$, $paramType3$, $paramName3$, $paramType4$, $paramName4$, $paramType5$, $paramName5$, $paramType6$, $paramName6$, $paramType7$, $paramName7$)$selected$ $end$]]></Code>
+        </Snippet>
+      </CodeSnippet>
+    </CodeSnippets>
\ No newline at end of file
diff --git a/Engine/Extras/VisualStudioSnippets/Unreal.delmul8.snippet b/Engine/Extras/VisualStudioSnippets/Unreal.delmul8.snippet
new file mode 100644
--- /dev/null
+++ b/Engine/Extras/VisualStudioSnippets/Unreal.delmul8.snippet
@@ -0,0 +1,140 @@
+<?xml version="1.0" encoding="utf-8"?>
+    <CodeSnippets xmlns="http://schemas.microsoft.com/VisualStudio/2005/CodeSnippet">
+      <CodeSnippet Format="1.0.0">
+        <Header>
+          <SnippetTypes>
+            <SnippetType>Expansion</SnippetType>
+          </SnippetTypes>
+          <Title>unreal-delmul8</Title>
+          <Author>Epic Games</Author>
+          <Description>Unreal Engine Multicast Delegate snippet</Description>
+          <HelpUrl>
+          </HelpUrl>
+          <Shortcut>unreal-delmul8</Shortcut>
+        </Header>
+        <Snippet>
+          <Declarations>
+       <Literal Editable="true">
+          <ID>delegateName</ID>
+          <ToolTip>Delegate type name</ToolTip>
+          <Default>DelegateName</Default>
+          <Function>
+          </Function>
+        </Literal>
+        <Literal Editable="true">
+          <ID>paramName1</ID>
+          <ToolTip>Parameter name</ToolTip>
+          <Default>Param1Name</Default>
+          <Function>
+          </Function>
+            </Literal>
+        <Literal Editable="true">
+          <ID>paramName2</ID>
+          <ToolTip>Parameter name</ToolTip>
+          <Default>Param2Name</Default>
+          <Function>
+          </Function>
+            </Literal>
+        <Literal Editable="true">
+          <ID>paramName3</ID>
+          <ToolTip>Parameter name</ToolTip>
+          <Default>Param3Name</Default>
+          <Function>
+          </Function>
+            </Literal>
+        <Literal Editable="true">
+          <ID>paramName4</ID>
+          <ToolTip>Parameter name</ToolTip>
+          <Default>Param4Name</Default>
+          <Function>
+          </Function>
+            </Literal>
+        <Literal Editable="true">
+          <ID>paramName5</ID>
+          <ToolTip>Parameter name</ToolTip>
+          <Default>Param5Name</Default>
+          <Function>
+          </Function>
+            </Literal>
+        <Literal Editable="true">
+          <ID>paramName6</ID>
+          <ToolTip>Parameter name</ToolTip>
+          <Default>Param6Name</Default>
+          <Function>
+          </Function>
+            </Literal>
+        <Literal Editable="true">
+          <ID>paramName7</ID>
+          <ToolTip>Parameter name</ToolTip>
+          <Default>Param7Name</Default>
+          <Function>
+          </Function>
+            </Literal>
+        <Literal Editable="true">
+          <ID>paramName8</ID>
+          <ToolTip>Parameter name</ToolTip>
+          <Default>Param8Name</Default>
+          <Function>
+          </Function>
+            </Literal>
+       <Literal Editable="true">
+          <ID>paramType1</ID>
+          <ToolTip>Parameter type</ToolTip>
+          <Default>Param1Type</Default>
+          <Function>
+          </Function>
+            </Literal>
+       <Literal Editable="true">
+          <ID>paramType2</ID>
+          <ToolTip>Parameter type</ToolTip>
+          <Default>Param2Type</Default>
+          <Function>
+          </Function>
+            </Literal>
+       <Literal Editable="true">
+          <ID>paramType3</ID>
+          <ToolTip>Parameter type</ToolTip>
+          <Default>Param3Type</Default>
+          <Function>
+          </Function>
+            </Literal>
+       <Literal Editable="true">
+          <ID>paramType4</ID>
+          <ToolTip>Parameter type</ToolTip>
+          <Default>Param4Type</Default>
+          <Function>
+          </Function>
+            </Literal>
+       <Literal Editable="true">
+          <ID>paramType5</ID>
+          <ToolTip>Parameter type</ToolTip>
+          <Default>Param5Type</Default>
+          <Function>
+          </Function>
+            </Literal>
+       <Literal Editable="true">
+          <ID>paramType6</ID>
+          <ToolTip>Parameter type</ToolTip>
+          <Default>Param6Type</Default>
+          <Function>
+          </Function>
+            </Literal>
+       <Literal Editable="true">
+          <ID>paramType7</ID>
+          <ToolTip>Parameter type</ToolTip>
+          <Default>Param7Type</Default>
+          <Function>
+          </Function>
+            </Literal>
+       <Literal Editable="true">
+          <ID>paramType8</ID>
+          <ToolTip>Parameter type</ToolTip>
+          <Default>Param8Type</Default>
+          <Function>
+          </Function>
+            </Literal>
+        </Declarations>
+          <Code Language="cpp"><![CDATA[DECLARE_MULTICAST_DELEGATE_EightParams($delegateName$, $paramType1$, $paramName1$, $paramType2$, $paramName2$, $paramType3$, $paramName3$, $paramType4$, $paramName4$, $paramType5$, $paramName5$, $paramType6$, $paramName6$, $paramType7$, $paramName7$, $paramType8$, $paramName8$)$selected$ $end$]]></Code>
+        </Snippet>
+      </CodeSnippet>
+    </CodeSnippets>
\ No newline at end of file
diff --git a/Engine/Extras/VisualStudioSnippets/Unreal.enum.snippet b/Engine/Extras/VisualStudioSnippets/Unreal.enum.snippet
new file mode 100644
--- /dev/null
+++ b/Engine/Extras/VisualStudioSnippets/Unreal.enum.snippet
@@ -0,0 +1,57 @@
+<?xml version="1.0" encoding="utf-8"?>
+<CodeSnippets xmlns="http://schemas.microsoft.com/VisualStudio/2005/CodeSnippet">
+  <CodeSnippet Format="1.0.0">
+    <Header>
+      <SnippetTypes>
+        <SnippetType>Expansion</SnippetType>
+      </SnippetTypes>
+      <Title>unreal-enum</Title>
+      <Author>Epic Games</Author>
+      <Description>Unreal Engine Enum snippet</Description>
+      <HelpUrl>
+      </HelpUrl>
+      <Shortcut>unreal-enum</Shortcut>
+    </Header>
+    <Snippet>
+      <Declarations>
+		<Literal Editable="true">
+          <ID>comment</ID>
+          <ToolTip>Enum comment</ToolTip>
+          <Default>Comment</Default>
+          <Function>
+          </Function>
+        </Literal>
+        <Literal Editable="true">
+          <ID>namespace</ID>
+          <ToolTip>Enum namespace name</ToolTip>
+          <Default>MyEnum</Default>
+          <Function>
+          </Function>
+        </Literal>
+		<Literal Editable="true">
+          <ID>elem1comment</ID>
+          <ToolTip>Enum element 1 comment</ToolTip>
+          <Default>Comment</Default>
+          <Function>
+          </Function>
+        </Literal>
+        <Literal Editable="true">
+          <ID>elem1name</ID>
+          <ToolTip>First enum element name</ToolTip>
+          <Default>Name</Default>
+          <Function>
+          </Function>
+        </Literal>
+      </Declarations>
+      <Code Language="cpp"><![CDATA[/**
+ * $comment$ 
+ */
+UENUM()
+enum E$namespace$
+{
+	$elem1name$, /**< $elem1comment$ */
+	$select$ $end$
+};]]></Code>
+    </Snippet>
+  </CodeSnippet>
+</CodeSnippets>
\ No newline at end of file
diff --git a/Engine/Extras/VisualStudioSnippets/Unreal.enumdisplay.snippet b/Engine/Extras/VisualStudioSnippets/Unreal.enumdisplay.snippet
new file mode 100644
--- /dev/null
+++ b/Engine/Extras/VisualStudioSnippets/Unreal.enumdisplay.snippet
@@ -0,0 +1,67 @@
+<?xml version="1.0" encoding="utf-8"?>
+<CodeSnippets xmlns="http://schemas.microsoft.com/VisualStudio/2005/CodeSnippet">
+  <CodeSnippet Format="1.0.0">
+    <Header>
+      <SnippetTypes>
+        <SnippetType>Expansion</SnippetType>
+      </SnippetTypes>
+      <Title>unreal-enumdisplay</Title>
+      <Author>Epic Games</Author>
+      <Description>Unreal Engine Enum snippet</Description>
+      <HelpUrl>
+      </HelpUrl>
+      <Shortcut>unreal-enumdisplay</Shortcut>
+    </Header>
+    <Snippet>
+      <Declarations>
+		<Literal Editable="true">
+          <ID>comment</ID>
+          <ToolTip>Enum comment</ToolTip>
+          <Default>Comment</Default>
+          <Function>
+          </Function>
+        </Literal>
+        <Literal Editable="true">
+          <ID>namespace</ID>
+          <ToolTip>Enum namespace name</ToolTip>
+          <Default>MyEnum</Default>
+          <Function>
+          </Function>
+        </Literal>
+        <Literal Editable="true">
+          <ID>elem1name</ID>
+          <ToolTip>First enum element name</ToolTip>
+          <Default>Name</Default>
+          <Function>
+          </Function>
+        </Literal>
+		<Literal Editable="true">
+          <ID>elem1comment</ID>
+          <ToolTip>Enum element 1 comment</ToolTip>
+          <Default>Comment</Default>
+          <Function>
+          </Function>
+        </Literal>
+		<Literal Editable="true">
+          <ID>elem1displayName</ID>
+          <ToolTip>First enum element display name</ToolTip>
+          <Default>DisplayName</Default>
+          <Function>
+          </Function>
+        </Literal>
+      </Declarations>
+      <Code Language="cpp"><![CDATA[/**
+ * $comment$ 
+ */
+UENUM()
+namespace E$namespace$
+{
+	enum Type
+	{
+		$elem1name$ UMETA(DisplayName="$elem1displayName$"), /**< $elem1comment$ */
+		$select$ $end$
+	};
+}]]></Code>
+    </Snippet>
+  </CodeSnippet>
+</CodeSnippets>
\ No newline at end of file
diff --git a/Engine/Extras/VisualStudioSnippets/Unreal.eve.snippet b/Engine/Extras/VisualStudioSnippets/Unreal.eve.snippet
new file mode 100644
--- /dev/null
+++ b/Engine/Extras/VisualStudioSnippets/Unreal.eve.snippet
@@ -0,0 +1,35 @@
+<?xml version="1.0" encoding="utf-8"?>
+    <CodeSnippets xmlns="http://schemas.microsoft.com/VisualStudio/2005/CodeSnippet">
+      <CodeSnippet Format="1.0.0">
+        <Header>
+          <SnippetTypes>
+            <SnippetType>Expansion</SnippetType>
+          </SnippetTypes>
+          <Title>unreal-eve</Title>
+          <Author>Epic Games</Author>
+          <Description>Unreal Engine Event snippet</Description>
+          <HelpUrl>
+          </HelpUrl>
+          <Shortcut>unreal-eve</Shortcut>
+        </Header>
+        <Snippet>
+          <Declarations>
+       <Literal Editable="true">
+          <ID>delegateName</ID>
+          <ToolTip>Delegate type name</ToolTip>
+          <Default>DelegateName</Default>
+          <Function>
+          </Function>
+        </Literal>
+       <Literal Editable="true">
+          <ID>owningType</ID>
+          <ToolTip>Delegate owning type name</ToolTip>
+          <Default>OwningType</Default>
+          <Function>
+          </Function>
+        </Literal>
+        </Declarations>
+          <Code Language="cpp"><![CDATA[DECLARE_EVENT($owningType$, $delegateName$)$selected$ $end$]]></Code>
+        </Snippet>
+      </CodeSnippet>
+    </CodeSnippets>
\ No newline at end of file
diff --git a/Engine/Extras/VisualStudioSnippets/Unreal.eve1.snippet b/Engine/Extras/VisualStudioSnippets/Unreal.eve1.snippet
new file mode 100644
--- /dev/null
+++ b/Engine/Extras/VisualStudioSnippets/Unreal.eve1.snippet
@@ -0,0 +1,42 @@
+<?xml version="1.0" encoding="utf-8"?>
+    <CodeSnippets xmlns="http://schemas.microsoft.com/VisualStudio/2005/CodeSnippet">
+      <CodeSnippet Format="1.0.0">
+        <Header>
+          <SnippetTypes>
+            <SnippetType>Expansion</SnippetType>
+          </SnippetTypes>
+          <Title>unreal-eve1</Title>
+          <Author>Epic Games</Author>
+          <Description>Unreal Engine Event snippet</Description>
+          <HelpUrl>
+          </HelpUrl>
+          <Shortcut>unreal-eve1</Shortcut>
+        </Header>
+        <Snippet>
+          <Declarations>
+       <Literal Editable="true">
+          <ID>delegateName</ID>
+          <ToolTip>Delegate type name</ToolTip>
+          <Default>DelegateName</Default>
+          <Function>
+          </Function>
+        </Literal>
+       <Literal Editable="true">
+          <ID>owningType</ID>
+          <ToolTip>Delegate owning type name</ToolTip>
+          <Default>OwningType</Default>
+          <Function>
+          </Function>
+        </Literal>
+       <Literal Editable="true">
+          <ID>paramType1</ID>
+          <ToolTip>Parameter type</ToolTip>
+          <Default>Param1Type</Default>
+          <Function>
+          </Function>
+            </Literal>
+        </Declarations>
+          <Code Language="cpp"><![CDATA[DECLARE_EVENT($owningType$, $delegateName$, $paramType1$)$selected$ $end$]]></Code>
+        </Snippet>
+      </CodeSnippet>
+    </CodeSnippets>
\ No newline at end of file
diff --git a/Engine/Extras/VisualStudioSnippets/Unreal.eve2.snippet b/Engine/Extras/VisualStudioSnippets/Unreal.eve2.snippet
new file mode 100644
--- /dev/null
+++ b/Engine/Extras/VisualStudioSnippets/Unreal.eve2.snippet
@@ -0,0 +1,49 @@
+<?xml version="1.0" encoding="utf-8"?>
+    <CodeSnippets xmlns="http://schemas.microsoft.com/VisualStudio/2005/CodeSnippet">
+      <CodeSnippet Format="1.0.0">
+        <Header>
+          <SnippetTypes>
+            <SnippetType>Expansion</SnippetType>
+          </SnippetTypes>
+          <Title>unreal-eve2</Title>
+          <Author>Epic Games</Author>
+          <Description>Unreal Engine Event snippet</Description>
+          <HelpUrl>
+          </HelpUrl>
+          <Shortcut>unreal-eve2</Shortcut>
+        </Header>
+        <Snippet>
+          <Declarations>
+       <Literal Editable="true">
+          <ID>delegateName</ID>
+          <ToolTip>Delegate type name</ToolTip>
+          <Default>DelegateName</Default>
+          <Function>
+          </Function>
+        </Literal>
+       <Literal Editable="true">
+          <ID>owningType</ID>
+          <ToolTip>Delegate owning type name</ToolTip>
+          <Default>OwningType</Default>
+          <Function>
+          </Function>
+        </Literal>
+       <Literal Editable="true">
+          <ID>paramType1</ID>
+          <ToolTip>Parameter type</ToolTip>
+          <Default>Param1Type</Default>
+          <Function>
+          </Function>
+            </Literal>
+       <Literal Editable="true">
+          <ID>paramType2</ID>
+          <ToolTip>Parameter type</ToolTip>
+          <Default>Param2Type</Default>
+          <Function>
+          </Function>
+            </Literal>
+        </Declarations>
+          <Code Language="cpp"><![CDATA[DECLARE_EVENT($owningType$, $delegateName$, $paramType1$, $paramType2$)$selected$ $end$]]></Code>
+        </Snippet>
+      </CodeSnippet>
+    </CodeSnippets>
\ No newline at end of file
diff --git a/Engine/Extras/VisualStudioSnippets/Unreal.eve3.snippet b/Engine/Extras/VisualStudioSnippets/Unreal.eve3.snippet
new file mode 100644
--- /dev/null
+++ b/Engine/Extras/VisualStudioSnippets/Unreal.eve3.snippet
@@ -0,0 +1,56 @@
+<?xml version="1.0" encoding="utf-8"?>
+    <CodeSnippets xmlns="http://schemas.microsoft.com/VisualStudio/2005/CodeSnippet">
+      <CodeSnippet Format="1.0.0">
+        <Header>
+          <SnippetTypes>
+            <SnippetType>Expansion</SnippetType>
+          </SnippetTypes>
+          <Title>unreal-eve3</Title>
+          <Author>Epic Games</Author>
+          <Description>Unreal Engine Event snippet</Description>
+          <HelpUrl>
+          </HelpUrl>
+          <Shortcut>unreal-eve3</Shortcut>
+        </Header>
+        <Snippet>
+          <Declarations>
+       <Literal Editable="true">
+          <ID>delegateName</ID>
+          <ToolTip>Delegate type name</ToolTip>
+          <Default>DelegateName</Default>
+          <Function>
+          </Function>
+        </Literal>
+       <Literal Editable="true">
+          <ID>owningType</ID>
+          <ToolTip>Delegate owning type name</ToolTip>
+          <Default>OwningType</Default>
+          <Function>
+          </Function>
+        </Literal>
+       <Literal Editable="true">
+          <ID>paramType1</ID>
+          <ToolTip>Parameter type</ToolTip>
+          <Default>Param1Type</Default>
+          <Function>
+          </Function>
+            </Literal>
+       <Literal Editable="true">
+          <ID>paramType2</ID>
+          <ToolTip>Parameter type</ToolTip>
+          <Default>Param2Type</Default>
+          <Function>
+          </Function>
+            </Literal>
+       <Literal Editable="true">
+          <ID>paramType3</ID>
+          <ToolTip>Parameter type</ToolTip>
+          <Default>Param3Type</Default>
+          <Function>
+          </Function>
+            </Literal>
+        </Declarations>
+          <Code Language="cpp"><![CDATA[DECLARE_EVENT($owningType$, $delegateName$, $paramType1$, $paramType2$, $paramType3$)$selected$ $end$]]></Code>
+        </Snippet>
+      </CodeSnippet>
+    </CodeSnippets>
\ No newline at end of file
diff --git a/Engine/Extras/VisualStudioSnippets/Unreal.eve4.snippet b/Engine/Extras/VisualStudioSnippets/Unreal.eve4.snippet
new file mode 100644
--- /dev/null
+++ b/Engine/Extras/VisualStudioSnippets/Unreal.eve4.snippet
@@ -0,0 +1,63 @@
+<?xml version="1.0" encoding="utf-8"?>
+    <CodeSnippets xmlns="http://schemas.microsoft.com/VisualStudio/2005/CodeSnippet">
+      <CodeSnippet Format="1.0.0">
+        <Header>
+          <SnippetTypes>
+            <SnippetType>Expansion</SnippetType>
+          </SnippetTypes>
+          <Title>unreal-eve4</Title>
+          <Author>Epic Games</Author>
+          <Description>Unreal Engine Event snippet</Description>
+          <HelpUrl>
+          </HelpUrl>
+          <Shortcut>unreal-eve4</Shortcut>
+        </Header>
+        <Snippet>
+          <Declarations>
+       <Literal Editable="true">
+          <ID>delegateName</ID>
+          <ToolTip>Delegate type name</ToolTip>
+          <Default>DelegateName</Default>
+          <Function>
+          </Function>
+        </Literal>
+       <Literal Editable="true">
+          <ID>owningType</ID>
+          <ToolTip>Delegate owning type name</ToolTip>
+          <Default>OwningType</Default>
+          <Function>
+          </Function>
+        </Literal>
+       <Literal Editable="true">
+          <ID>paramType1</ID>
+          <ToolTip>Parameter type</ToolTip>
+          <Default>Param1Type</Default>
+          <Function>
+          </Function>
+            </Literal>
+       <Literal Editable="true">
+          <ID>paramType2</ID>
+          <ToolTip>Parameter type</ToolTip>
+          <Default>Param2Type</Default>
+          <Function>
+          </Function>
+            </Literal>
+       <Literal Editable="true">
+          <ID>paramType3</ID>
+          <ToolTip>Parameter type</ToolTip>
+          <Default>Param3Type</Default>
+          <Function>
+          </Function>
+            </Literal>
+       <Literal Editable="true">
+          <ID>paramType4</ID>
+          <ToolTip>Parameter type</ToolTip>
+          <Default>Param4Type</Default>
+          <Function>
+          </Function>
+            </Literal>
+        </Declarations>
+          <Code Language="cpp"><![CDATA[DECLARE_EVENT($owningType$, $delegateName$, $paramType1$, $paramType2$, $paramType3$, $paramType4$)$selected$ $end$]]></Code>
+        </Snippet>
+      </CodeSnippet>
+    </CodeSnippets>
\ No newline at end of file
diff --git a/Engine/Extras/VisualStudioSnippets/Unreal.eve5.snippet b/Engine/Extras/VisualStudioSnippets/Unreal.eve5.snippet
new file mode 100644
--- /dev/null
+++ b/Engine/Extras/VisualStudioSnippets/Unreal.eve5.snippet
@@ -0,0 +1,70 @@
+<?xml version="1.0" encoding="utf-8"?>
+    <CodeSnippets xmlns="http://schemas.microsoft.com/VisualStudio/2005/CodeSnippet">
+      <CodeSnippet Format="1.0.0">
+        <Header>
+          <SnippetTypes>
+            <SnippetType>Expansion</SnippetType>
+          </SnippetTypes>
+          <Title>unreal-eve5</Title>
+          <Author>Epic Games</Author>
+          <Description>Unreal Engine Event snippet</Description>
+          <HelpUrl>
+          </HelpUrl>
+          <Shortcut>unreal-eve5</Shortcut>
+        </Header>
+        <Snippet>
+          <Declarations>
+       <Literal Editable="true">
+          <ID>delegateName</ID>
+          <ToolTip>Delegate type name</ToolTip>
+          <Default>DelegateName</Default>
+          <Function>
+          </Function>
+        </Literal>
+       <Literal Editable="true">
+          <ID>owningType</ID>
+          <ToolTip>Delegate owning type name</ToolTip>
+          <Default>OwningType</Default>
+          <Function>
+          </Function>
+        </Literal>
+       <Literal Editable="true">
+          <ID>paramType1</ID>
+          <ToolTip>Parameter type</ToolTip>
+          <Default>Param1Type</Default>
+          <Function>
+          </Function>
+            </Literal>
+       <Literal Editable="true">
+          <ID>paramType2</ID>
+          <ToolTip>Parameter type</ToolTip>
+          <Default>Param2Type</Default>
+          <Function>
+          </Function>
+            </Literal>
+       <Literal Editable="true">
+          <ID>paramType3</ID>
+          <ToolTip>Parameter type</ToolTip>
+          <Default>Param3Type</Default>
+          <Function>
+          </Function>
+            </Literal>
+       <Literal Editable="true">
+          <ID>paramType4</ID>
+          <ToolTip>Parameter type</ToolTip>
+          <Default>Param4Type</Default>
+          <Function>
+          </Function>
+            </Literal>
+       <Literal Editable="true">
+          <ID>paramType5</ID>
+          <ToolTip>Parameter type</ToolTip>
+          <Default>Param5Type</Default>
+          <Function>
+          </Function>
+            </Literal>
+        </Declarations>
+          <Code Language="cpp"><![CDATA[DECLARE_EVENT($owningType$, $delegateName$, $paramType1$, $paramType2$, $paramType3$, $paramType4$, $paramType5$)$selected$ $end$]]></Code>
+        </Snippet>
+      </CodeSnippet>
+    </CodeSnippets>
\ No newline at end of file
diff --git a/Engine/Extras/VisualStudioSnippets/Unreal.eve6.snippet b/Engine/Extras/VisualStudioSnippets/Unreal.eve6.snippet
new file mode 100644
--- /dev/null
+++ b/Engine/Extras/VisualStudioSnippets/Unreal.eve6.snippet
@@ -0,0 +1,77 @@
+<?xml version="1.0" encoding="utf-8"?>
+    <CodeSnippets xmlns="http://schemas.microsoft.com/VisualStudio/2005/CodeSnippet">
+      <CodeSnippet Format="1.0.0">
+        <Header>
+          <SnippetTypes>
+            <SnippetType>Expansion</SnippetType>
+          </SnippetTypes>
+          <Title>unreal-eve6</Title>
+          <Author>Epic Games</Author>
+          <Description>Unreal Engine Event snippet</Description>
+          <HelpUrl>
+          </HelpUrl>
+          <Shortcut>unreal-eve6</Shortcut>
+        </Header>
+        <Snippet>
+          <Declarations>
+       <Literal Editable="true">
+          <ID>delegateName</ID>
+          <ToolTip>Delegate type name</ToolTip>
+          <Default>DelegateName</Default>
+          <Function>
+          </Function>
+        </Literal>
+       <Literal Editable="true">
+          <ID>owningType</ID>
+          <ToolTip>Delegate owning type name</ToolTip>
+          <Default>OwningType</Default>
+          <Function>
+          </Function>
+        </Literal>
+       <Literal Editable="true">
+          <ID>paramType1</ID>
+          <ToolTip>Parameter type</ToolTip>
+          <Default>Param1Type</Default>
+          <Function>
+          </Function>
+            </Literal>
+       <Literal Editable="true">
+          <ID>paramType2</ID>
+          <ToolTip>Parameter type</ToolTip>
+          <Default>Param2Type</Default>
+          <Function>
+          </Function>
+            </Literal>
+       <Literal Editable="true">
+          <ID>paramType3</ID>
+          <ToolTip>Parameter type</ToolTip>
+          <Default>Param3Type</Default>
+          <Function>
+          </Function>
+            </Literal>
+       <Literal Editable="true">
+          <ID>paramType4</ID>
+          <ToolTip>Parameter type</ToolTip>
+          <Default>Param4Type</Default>
+          <Function>
+          </Function>
+            </Literal>
+       <Literal Editable="true">
+          <ID>paramType5</ID>
+          <ToolTip>Parameter type</ToolTip>
+          <Default>Param5Type</Default>
+          <Function>
+          </Function>
+            </Literal>
+       <Literal Editable="true">
+          <ID>paramType6</ID>
+          <ToolTip>Parameter type</ToolTip>
+          <Default>Param6Type</Default>
+          <Function>
+          </Function>
+            </Literal>
+        </Declarations>
+          <Code Language="cpp"><![CDATA[DECLARE_EVENT($owningType$, $delegateName$, $paramType1$, $paramType2$, $paramType3$, $paramType4$, $paramType5$, $paramType6$)$selected$ $end$]]></Code>
+        </Snippet>
+      </CodeSnippet>
+    </CodeSnippets>
\ No newline at end of file
diff --git a/Engine/Extras/VisualStudioSnippets/Unreal.eve7.snippet b/Engine/Extras/VisualStudioSnippets/Unreal.eve7.snippet
new file mode 100644
--- /dev/null
+++ b/Engine/Extras/VisualStudioSnippets/Unreal.eve7.snippet
@@ -0,0 +1,84 @@
+<?xml version="1.0" encoding="utf-8"?>
+    <CodeSnippets xmlns="http://schemas.microsoft.com/VisualStudio/2005/CodeSnippet">
+      <CodeSnippet Format="1.0.0">
+        <Header>
+          <SnippetTypes>
+            <SnippetType>Expansion</SnippetType>
+          </SnippetTypes>
+          <Title>unreal-eve7</Title>
+          <Author>Epic Games</Author>
+          <Description>Unreal Engine Event snippet</Description>
+          <HelpUrl>
+          </HelpUrl>
+          <Shortcut>unreal-eve7</Shortcut>
+        </Header>
+        <Snippet>
+          <Declarations>
+       <Literal Editable="true">
+          <ID>delegateName</ID>
+          <ToolTip>Delegate type name</ToolTip>
+          <Default>DelegateName</Default>
+          <Function>
+          </Function>
+        </Literal>
+       <Literal Editable="true">
+          <ID>owningType</ID>
+          <ToolTip>Delegate owning type name</ToolTip>
+          <Default>OwningType</Default>
+          <Function>
+          </Function>
+        </Literal>
+       <Literal Editable="true">
+          <ID>paramType1</ID>
+          <ToolTip>Parameter type</ToolTip>
+          <Default>Param1Type</Default>
+          <Function>
+          </Function>
+            </Literal>
+       <Literal Editable="true">
+          <ID>paramType2</ID>
+          <ToolTip>Parameter type</ToolTip>
+          <Default>Param2Type</Default>
+          <Function>
+          </Function>
+            </Literal>
+       <Literal Editable="true">
+          <ID>paramType3</ID>
+          <ToolTip>Parameter type</ToolTip>
+          <Default>Param3Type</Default>
+          <Function>
+          </Function>
+            </Literal>
+       <Literal Editable="true">
+          <ID>paramType4</ID>
+          <ToolTip>Parameter type</ToolTip>
+          <Default>Param4Type</Default>
+          <Function>
+          </Function>
+            </Literal>
+       <Literal Editable="true">
+          <ID>paramType5</ID>
+          <ToolTip>Parameter type</ToolTip>
+          <Default>Param5Type</Default>
+          <Function>
+          </Function>
+            </Literal>
+       <Literal Editable="true">
+          <ID>paramType6</ID>
+          <ToolTip>Parameter type</ToolTip>
+          <Default>Param6Type</Default>
+          <Function>
+          </Function>
+            </Literal>
+       <Literal Editable="true">
+          <ID>paramType7</ID>
+          <ToolTip>Parameter type</ToolTip>
+          <Default>Param7Type</Default>
+          <Function>
+          </Function>
+            </Literal>
+        </Declarations>
+          <Code Language="cpp"><![CDATA[DECLARE_EVENT($owningType$, $delegateName$, $paramType1$, $paramType2$, $paramType3$, $paramType4$, $paramType5$, $paramType6$, $paramType7$)$selected$ $end$]]></Code>
+        </Snippet>
+      </CodeSnippet>
+    </CodeSnippets>
\ No newline at end of file
diff --git a/Engine/Extras/VisualStudioSnippets/Unreal.eve8.snippet b/Engine/Extras/VisualStudioSnippets/Unreal.eve8.snippet
new file mode 100644
--- /dev/null
+++ b/Engine/Extras/VisualStudioSnippets/Unreal.eve8.snippet
@@ -0,0 +1,91 @@
+<?xml version="1.0" encoding="utf-8"?>
+    <CodeSnippets xmlns="http://schemas.microsoft.com/VisualStudio/2005/CodeSnippet">
+      <CodeSnippet Format="1.0.0">
+        <Header>
+          <SnippetTypes>
+            <SnippetType>Expansion</SnippetType>
+          </SnippetTypes>
+          <Title>unreal-eve8</Title>
+          <Author>Epic Games</Author>
+          <Description>Unreal Engine Event snippet</Description>
+          <HelpUrl>
+          </HelpUrl>
+          <Shortcut>unreal-eve8</Shortcut>
+        </Header>
+        <Snippet>
+          <Declarations>
+       <Literal Editable="true">
+          <ID>delegateName</ID>
+          <ToolTip>Delegate type name</ToolTip>
+          <Default>DelegateName</Default>
+          <Function>
+          </Function>
+        </Literal>
+       <Literal Editable="true">
+          <ID>owningType</ID>
+          <ToolTip>Delegate owning type name</ToolTip>
+          <Default>OwningType</Default>
+          <Function>
+          </Function>
+        </Literal>
+       <Literal Editable="true">
+          <ID>paramType1</ID>
+          <ToolTip>Parameter type</ToolTip>
+          <Default>Param1Type</Default>
+          <Function>
+          </Function>
+            </Literal>
+       <Literal Editable="true">
+          <ID>paramType2</ID>
+          <ToolTip>Parameter type</ToolTip>
+          <Default>Param2Type</Default>
+          <Function>
+          </Function>
+            </Literal>
+       <Literal Editable="true">
+          <ID>paramType3</ID>
+          <ToolTip>Parameter type</ToolTip>
+          <Default>Param3Type</Default>
+          <Function>
+          </Function>
+            </Literal>
+       <Literal Editable="true">
+          <ID>paramType4</ID>
+          <ToolTip>Parameter type</ToolTip>
+          <Default>Param4Type</Default>
+          <Function>
+          </Function>
+            </Literal>
+       <Literal Editable="true">
+          <ID>paramType5</ID>
+          <ToolTip>Parameter type</ToolTip>
+          <Default>Param5Type</Default>
+          <Function>
+          </Function>
+            </Literal>
+       <Literal Editable="true">
+          <ID>paramType6</ID>
+          <ToolTip>Parameter type</ToolTip>
+          <Default>Param6Type</Default>
+          <Function>
+          </Function>
+            </Literal>
+       <Literal Editable="true">
+          <ID>paramType7</ID>
+          <ToolTip>Parameter type</ToolTip>
+          <Default>Param7Type</Default>
+          <Function>
+          </Function>
+            </Literal>
+       <Literal Editable="true">
+          <ID>paramType8</ID>
+          <ToolTip>Parameter type</ToolTip>
+          <Default>Param8Type</Default>
+          <Function>
+          </Function>
+            </Literal>
+        </Declarations>
+          <Code Language="cpp"><![CDATA[DECLARE_EVENT($owningType$, $delegateName$, $paramType1$, $paramType2$, $paramType3$, $paramType4$, $paramType5$, $paramType6$, $paramType7$, $paramType8$)$selected$ $end$]]></Code>
+        </Snippet>
+      </CodeSnippet>
+    </CodeSnippets>
\ No newline at end of file
diff --git a/Engine/Extras/VisualStudioSnippets/Unreal.interface.snippet b/Engine/Extras/VisualStudioSnippets/Unreal.interface.snippet
new file mode 100644
--- /dev/null
+++ b/Engine/Extras/VisualStudioSnippets/Unreal.interface.snippet
@@ -0,0 +1,49 @@
+<?xml version="1.0" encoding="utf-8"?>
+<CodeSnippets xmlns="http://schemas.microsoft.com/VisualStudio/2005/CodeSnippet">
+  <CodeSnippet Format="1.0.0">
+    <Header>
+      <SnippetTypes>
+        <SnippetType>Expansion</SnippetType>
+      </SnippetTypes>
+      <Title>unreal-interface</Title>
+      <Author>Epic Games</Author>
+      <Description>Unreal Engine interface snippet</Description>
+      <HelpUrl>
+      </HelpUrl>
+      <Shortcut>unreal-interface</Shortcut>
+    </Header>
+    <Snippet>
+      <Declarations>
+        <Literal Editable="true">
+          <ID>comment</ID>
+          <ToolTip>Interface comment</ToolTip>
+          <Default>Comment</Default>
+          <Function>
+          </Function>
+        </Literal>
+        <Literal Editable="true">
+          <ID>name</ID>
+          <ToolTip>Interface name</ToolTip>
+          <Default>Foo</Default>
+          <Function>
+          </Function>
+        </Literal>
+      </Declarations>
+      <Code Language="cpp"><![CDATA[/**
+ * $comment$
+ */
+UINTERFACE()
+class U$name$ : public UInterface
+{
+	GENERATED_UINTERFACE_BODY()
+};
+
+class I$name$
+{
+	GENERATED_IINTERFACE_BODY()
+  
+	$selected$ $end$
+};]]></Code>
+    </Snippet>
+  </CodeSnippet>
+</CodeSnippets>
\ No newline at end of file
diff --git a/Engine/Extras/VisualStudioSnippets/Unreal.log.snippet b/Engine/Extras/VisualStudioSnippets/Unreal.log.snippet
new file mode 100644
--- /dev/null
+++ b/Engine/Extras/VisualStudioSnippets/Unreal.log.snippet
@@ -0,0 +1,43 @@
+<?xml version="1.0" encoding="utf-8"?>
+<CodeSnippets xmlns="http://schemas.microsoft.com/VisualStudio/2005/CodeSnippet">
+  <CodeSnippet Format="1.0.0">
+    <Header>
+      <SnippetTypes>
+        <SnippetType>Expansion</SnippetType>
+      </SnippetTypes>
+      <Title>unreal-log</Title>
+      <Author>Epic Games</Author>
+      <Description>Unreal Engine UE_LOG snippet</Description>
+      <HelpUrl>
+      </HelpUrl>
+      <Shortcut>unreal-log</Shortcut>
+    </Header>
+    <Snippet>
+      <Declarations>
+        <Literal Editable="true">
+          <ID>category</ID>
+          <ToolTip>Log category name</ToolTip>
+          <Default>MyProject</Default>
+          <Function>
+          </Function>
+        </Literal>
+        <Literal Editable="true">
+          <ID>level</ID>
+          <ToolTip>Log level (NoLoggigng, Fatal, Error, Warning, Display, Log, Verbose, VeryVerbose, All, NumVerbosity, VerbosityMask, SetColor, BreakOnLog)</ToolTip>
+          <Default>Error</Default>
+          <Function>
+          </Function>
+        </Literal>
+        <Literal Editable="true">
+          <ID>text</ID>
+          <ToolTip>Enter log message here</ToolTip>
+          <Default>Log message</Default>
+          <Function>
+          </Function>
+        </Literal>
+      </Declarations>
+      <Code Language="cpp"><![CDATA[UE_LOG($category$, $level$, TEXT("$text$"));
+$selected$ $end$]]></Code>
+    </Snippet>
+  </CodeSnippet>
+</CodeSnippets>
\ No newline at end of file
diff --git a/Engine/Extras/VisualStudioSnippets/Unreal.logdeclare.snippet b/Engine/Extras/VisualStudioSnippets/Unreal.logdeclare.snippet
new file mode 100644
--- /dev/null
+++ b/Engine/Extras/VisualStudioSnippets/Unreal.logdeclare.snippet
@@ -0,0 +1,43 @@
+<?xml version="1.0" encoding="utf-8"?>
+<CodeSnippets xmlns="http://schemas.microsoft.com/VisualStudio/2005/CodeSnippet">
+  <CodeSnippet Format="1.0.0">
+    <Header>
+      <SnippetTypes>
+        <SnippetType>Expansion</SnippetType>
+      </SnippetTypes>
+      <Title>unreal-logdeclare</Title>
+      <Author>Epic Games</Author>
+      <Description>Unreal Engine UE_LOG snippet</Description>
+      <HelpUrl>
+      </HelpUrl>
+      <Shortcut>unreal-logdeclare</Shortcut>
+    </Header>
+    <Snippet>
+      <Declarations>
+        <Literal Editable="true">
+          <ID>category</ID>
+          <ToolTip>Log category name</ToolTip>
+          <Default>MyProject</Default>
+          <Function>
+          </Function>
+        </Literal>
+        <Literal Editable="true">
+          <ID>defaultVerbosity</ID>
+          <ToolTip>Default verbosity (NoLoggigng, Fatal, Error, Warning, Display, Log, Verbose, VeryVerbose, All, NumVerbosity, VerbosityMask, SetColor, BreakOnLog)</ToolTip>
+          <Default>Log</Default>
+          <Function>
+          </Function>
+        </Literal>
+        <Literal Editable="true">
+          <ID>compileTimeVerbosity</ID>
+          <ToolTip>Compile time verbosity</ToolTip>
+          <Default>All</Default>
+          <Function>
+          </Function>
+        </Literal>
+      </Declarations>
+      <Code Language="cpp"><![CDATA[DECLARE_LOG_CATEGORY_EXTERN($category$, $defaultVerbosity$, $compileTimeVerbosity$);
+$selected$ $end$]]></Code>
+    </Snippet>
+  </CodeSnippet>
+</CodeSnippets>
\ No newline at end of file
diff --git a/Engine/Extras/VisualStudioSnippets/Unreal.logdefine.snippet b/Engine/Extras/VisualStudioSnippets/Unreal.logdefine.snippet
new file mode 100644
--- /dev/null
+++ b/Engine/Extras/VisualStudioSnippets/Unreal.logdefine.snippet
@@ -0,0 +1,29 @@
+<?xml version="1.0" encoding="utf-8"?>
+<CodeSnippets xmlns="http://schemas.microsoft.com/VisualStudio/2005/CodeSnippet">
+  <CodeSnippet Format="1.0.0">
+    <Header>
+      <SnippetTypes>
+        <SnippetType>Expansion</SnippetType>
+      </SnippetTypes>
+      <Title>unreal-logdefine</Title>
+      <Author>Epic Games</Author>
+      <Description>Unreal Engine UE_LOG snippet</Description>
+      <HelpUrl>
+      </HelpUrl>
+      <Shortcut>unreal-logdefine</Shortcut>
+    </Header>
+    <Snippet>
+      <Declarations>
+        <Literal Editable="true">
+          <ID>category</ID>
+          <ToolTip>Log category name</ToolTip>
+          <Default>MyProject</Default>
+          <Function>
+          </Function>
+        </Literal>
+      </Declarations>
+      <Code Language="cpp"><![CDATA[DEFINE_LOG_CATEGORY($category$);
+$selected$ $end$]]></Code>
+    </Snippet>
+  </CodeSnippet>
+</CodeSnippets>
\ No newline at end of file
diff --git a/Engine/Extras/VisualStudioSnippets/Unreal.logfloat.snippet b/Engine/Extras/VisualStudioSnippets/Unreal.logfloat.snippet
new file mode 100644
--- /dev/null
+++ b/Engine/Extras/VisualStudioSnippets/Unreal.logfloat.snippet
@@ -0,0 +1,43 @@
+<?xml version="1.0" encoding="utf-8"?>
+<CodeSnippets xmlns="http://schemas.microsoft.com/VisualStudio/2005/CodeSnippet">
+  <CodeSnippet Format="1.0.0">
+    <Header>
+      <SnippetTypes>
+        <SnippetType>Expansion</SnippetType>
+      </SnippetTypes>
+      <Title>unreal-logfloat</Title>
+      <Author>Epic Games</Author>
+      <Description>Unreal Engine UE_LOG snippet</Description>
+      <HelpUrl>
+      </HelpUrl>
+      <Shortcut>unreal-logfloat</Shortcut>
+    </Header>
+    <Snippet>
+      <Declarations>
+        <Literal Editable="true">
+          <ID>category</ID>
+          <ToolTip>Log category name</ToolTip>
+          <Default>MyProject</Default>
+          <Function>
+          </Function>
+        </Literal>
+        <Literal Editable="true">
+          <ID>level</ID>
+          <ToolTip>Log level (NoLoggigng, Fatal, Error, Warning, Display, Log, Verbose, VeryVerbose, All, NumVerbosity, VerbosityMask, SetColor, BreakOnLog)</ToolTip>
+          <Default>Error</Default>
+          <Function>
+          </Function>
+        </Literal>
+        <Literal Editable="true">
+          <ID>variableName</ID>
+          <ToolTip>Name of variable</ToolTip>
+          <Default>variable</Default>
+          <Function>
+          </Function>
+        </Literal>
+      </Declarations>
+      <Code Language="cpp"><![CDATA[UE_LOG($category$, $level$, TEXT("The value of '$variableName$' is: %f"), $variableName$);
+$selected$ $end$]]></Code>
+    </Snippet>
+  </CodeSnippet>
+</CodeSnippets>
\ No newline at end of file
diff --git a/Engine/Extras/VisualStudioSnippets/Unreal.logint.snippet b/Engine/Extras/VisualStudioSnippets/Unreal.logint.snippet
new file mode 100644
--- /dev/null
+++ b/Engine/Extras/VisualStudioSnippets/Unreal.logint.snippet
@@ -0,0 +1,43 @@
+<?xml version="1.0" encoding="utf-8"?>
+<CodeSnippets xmlns="http://schemas.microsoft.com/VisualStudio/2005/CodeSnippet">
+  <CodeSnippet Format="1.0.0">
+    <Header>
+      <SnippetTypes>
+        <SnippetType>Expansion</SnippetType>
+      </SnippetTypes>
+      <Title>unreal-logint</Title>
+      <Author>Epic Games</Author>
+      <Description>Unreal Engine UE_LOG snippet</Description>
+      <HelpUrl>
+      </HelpUrl>
+      <Shortcut>unreal-logint</Shortcut>
+    </Header>
+    <Snippet>
+      <Declarations>
+        <Literal Editable="true">
+          <ID>category</ID>
+          <ToolTip>Log category name</ToolTip>
+          <Default>MyProject</Default>
+          <Function>
+          </Function>
+        </Literal>
+        <Literal Editable="true">
+          <ID>level</ID>
+          <ToolTip>Log level (NoLoggigng, Fatal, Error, Warning, Display, Log, Verbose, VeryVerbose, All, NumVerbosity, VerbosityMask, SetColor, BreakOnLog)</ToolTip>
+          <Default>Error</Default>
+          <Function>
+          </Function>
+        </Literal>
+        <Literal Editable="true">
+          <ID>variableName</ID>
+          <ToolTip>Name of variable</ToolTip>
+          <Default>variable</Default>
+          <Function>
+          </Function>
+        </Literal>
+      </Declarations>
+      <Code Language="cpp"><![CDATA[UE_LOG($category$, $level$, TEXT("The value of '$variableName$' is: %i"), $variableName$);
+$selected$ $end$]]></Code>
+    </Snippet>
+  </CodeSnippet>
+</CodeSnippets>
\ No newline at end of file
diff --git a/Engine/Extras/VisualStudioSnippets/Unreal.loguobj.snippet b/Engine/Extras/VisualStudioSnippets/Unreal.loguobj.snippet
new file mode 100644
--- /dev/null
+++ b/Engine/Extras/VisualStudioSnippets/Unreal.loguobj.snippet
@@ -0,0 +1,50 @@
+<?xml version="1.0" encoding="utf-8"?>
+<CodeSnippets xmlns="http://schemas.microsoft.com/VisualStudio/2005/CodeSnippet">
+  <CodeSnippet Format="1.0.0">
+    <Header>
+      <SnippetTypes>
+        <SnippetType>Expansion</SnippetType>
+      </SnippetTypes>
+      <Title>unreal-loguobj</Title>
+      <Author>Epic Games</Author>
+      <Description>Unreal Engine UE_LOG snippet</Description>
+      <HelpUrl>
+      </HelpUrl>
+      <Shortcut>unreal-loguobj</Shortcut>
+    </Header>
+    <Snippet>
+      <Declarations>
+        <Literal Editable="true">
+          <ID>category</ID>
+          <ToolTip>Log category name</ToolTip>
+          <Default>MyProject</Default>
+          <Function>
+          </Function>
+        </Literal>
+        <Literal Editable="true">
+          <ID>level</ID>
+          <ToolTip>Log level (NoLoggigng, Fatal, Error, Warning, Display, Log, Verbose, VeryVerbose, All, NumVerbosity, VerbosityMask, SetColor, BreakOnLog)</ToolTip>
+          <Default>Error</Default>
+          <Function>
+          </Function>
+        </Literal>
+        <Literal Editable="true">
+          <ID>text</ID>
+          <ToolTip>Enter log message here</ToolTip>
+          <Default>Log message</Default>
+          <Function>
+          </Function>
+        </Literal>
+		<Literal Editable="true">
+          <ID>object</ID>
+          <ToolTip>UObject object instance</ToolTip>
+          <Default>this</Default>
+          <Function>
+          </Function>
+        </Literal>
+      </Declarations>
+      <Code Language="cpp"><![CDATA[UE_LOG($category$, $level$, TEXT("[%s] $text$"), *($object$->GetName()));
+$selected$ $end$]]></Code>
+    </Snippet>
+  </CodeSnippet>
+</CodeSnippets>
\ No newline at end of file
diff --git a/Engine/Extras/VisualStudioSnippets/Unreal.mark.snippet b/Engine/Extras/VisualStudioSnippets/Unreal.mark.snippet
new file mode 100644
--- /dev/null
+++ b/Engine/Extras/VisualStudioSnippets/Unreal.mark.snippet
@@ -0,0 +1,62 @@
+<?xml version="1.0" encoding="utf-8"?>
+<CodeSnippets xmlns="http://schemas.microsoft.com/VisualStudio/2005/CodeSnippet">
+  <CodeSnippet Format="1.0.0">
+    <Header>
+      <SnippetTypes>
+        <SnippetType>Expansion</SnippetType>
+      </SnippetTypes>
+      <Title>unreal-mark</Title>
+      <Author>Epic Games</Author>
+      <Description>Unreal Engine mark snippet</Description>
+      <HelpUrl>
+      </HelpUrl>
+      <Shortcut>unreal-mark</Shortcut>
+    </Header>
+    <Snippet>
+      <Declarations>
+        <Literal Editable="true">
+          <ID>companySymbol</ID>
+          <ToolTip>Symbol of your company</ToolTip>
+          <Default>MyCompany</Default>
+          <Function>
+          </Function>
+        </Literal>
+		<Literal Editable="true">
+          <ID>task</ID>
+          <ToolTip>Task/ticket number</ToolTip>
+          <Default>Ticket number</Default>
+          <Function>
+          </Function>
+        </Literal>
+		<Literal Editable="true">
+          <ID>name</ID>
+          <ToolTip>Developer name</ToolTip>
+          <Default>Name</Default>
+          <Function>
+          </Function>
+        </Literal>
+		<Literal Editable="true">
+          <ID>surname</ID>
+          <ToolTip>Developer surname</ToolTip>
+          <Default>Surname</Default>
+          <Function>
+          </Function>
+        </Literal>
+		<Literal Editable="true">
+          <ID>description</ID>
+          <ToolTip>Description of modification</ToolTip>
+          <Default>Description</Default>
+          <Function>
+          </Function>
+        </Literal>
+      </Declarations>
+      <Code Language="cpp"><![CDATA[/* BEGIN $companySymbol$ ($task$) */
+	// $name$ $surname$ - $description$
+
+	$selected$ $end$
+
+	/* END $companySymbol$ */
+]]></Code>
+    </Snippet>
+  </CodeSnippet>
+</CodeSnippets>
\ No newline at end of file
diff --git a/Engine/Extras/VisualStudioSnippets/Unreal.struct.snippet b/Engine/Extras/VisualStudioSnippets/Unreal.struct.snippet
new file mode 100644
--- /dev/null
+++ b/Engine/Extras/VisualStudioSnippets/Unreal.struct.snippet
@@ -0,0 +1,45 @@
+<?xml version="1.0" encoding="utf-8"?>
+<CodeSnippets xmlns="http://schemas.microsoft.com/VisualStudio/2005/CodeSnippet">
+  <CodeSnippet Format="1.0.0">
+    <Header>
+      <SnippetTypes>
+        <SnippetType>Expansion</SnippetType>
+      </SnippetTypes>
+      <Title>unreal-struct</Title>
+      <Author>Epic Games</Author>
+      <Description>Unreal Engine struct snippet</Description>
+      <HelpUrl>
+      </HelpUrl>
+      <Shortcut>unreal-struct</Shortcut>
+    </Header>
+    <Snippet>
+      <Declarations>
+        <Literal Editable="true">
+          <ID>comment</ID>
+          <ToolTip>Struct comment</ToolTip>
+          <Default>Comment</Default>
+          <Function>
+          </Function>
+        </Literal>
+		<Literal Editable="true">
+          <ID>name</ID>
+          <ToolTip>Struct name</ToolTip>
+          <Default>Foo</Default>
+          <Function>
+          </Function>
+        </Literal>
+      </Declarations>
+      <Code Language="cpp"><![CDATA[/**
+ * $comment$ 
+ */
+USTRUCT()
+struct F$name$
+{
+	GENERATED_USTRUCT_BODY()
+
+	$selected$ $end$
+};
+]]></Code>
+    </Snippet>
+  </CodeSnippet>
+</CodeSnippets>
\ No newline at end of file
diff --git a/Engine/Extras/VisualStudioSnippets/ue4bpevent.snippet b/Engine/Extras/VisualStudioSnippets/ue4bpevent.snippet
deleted file mode 100644
diff --git a/Engine/Extras/VisualStudioSnippets/ue4bpfunc.snippet b/Engine/Extras/VisualStudioSnippets/ue4bpfunc.snippet
deleted file mode 100644
diff --git a/Engine/Extras/VisualStudioSnippets/ue4bpprop.snippet b/Engine/Extras/VisualStudioSnippets/ue4bpprop.snippet
deleted file mode 100644
diff --git a/Engine/Extras/VisualStudioSnippets/ue4classa.snippet b/Engine/Extras/VisualStudioSnippets/ue4classa.snippet
deleted file mode 100644
diff --git a/Engine/Extras/VisualStudioSnippets/ue4classu.snippet b/Engine/Extras/VisualStudioSnippets/ue4classu.snippet
deleted file mode 100644
diff --git a/Engine/Extras/VisualStudioSnippets/ue4del.snippet b/Engine/Extras/VisualStudioSnippets/ue4del.snippet
deleted file mode 100644
diff --git a/Engine/Extras/VisualStudioSnippets/ue4del1.snippet b/Engine/Extras/VisualStudioSnippets/ue4del1.snippet
deleted file mode 100644
diff --git a/Engine/Extras/VisualStudioSnippets/ue4del2.snippet b/Engine/Extras/VisualStudioSnippets/ue4del2.snippet
deleted file mode 100644
diff --git a/Engine/Extras/VisualStudioSnippets/ue4del3.snippet b/Engine/Extras/VisualStudioSnippets/ue4del3.snippet
deleted file mode 100644
diff --git a/Engine/Extras/VisualStudioSnippets/ue4del4.snippet b/Engine/Extras/VisualStudioSnippets/ue4del4.snippet
deleted file mode 100644
diff --git a/Engine/Extras/VisualStudioSnippets/ue4del5.snippet b/Engine/Extras/VisualStudioSnippets/ue4del5.snippet
deleted file mode 100644
diff --git a/Engine/Extras/VisualStudioSnippets/ue4del6.snippet b/Engine/Extras/VisualStudioSnippets/ue4del6.snippet
deleted file mode 100644
diff --git a/Engine/Extras/VisualStudioSnippets/ue4del7.snippet b/Engine/Extras/VisualStudioSnippets/ue4del7.snippet
deleted file mode 100644
diff --git a/Engine/Extras/VisualStudioSnippets/ue4del8.snippet b/Engine/Extras/VisualStudioSnippets/ue4del8.snippet
deleted file mode 100644
diff --git a/Engine/Extras/VisualStudioSnippets/ue4deldyn.snippet b/Engine/Extras/VisualStudioSnippets/ue4deldyn.snippet
deleted file mode 100644
diff --git a/Engine/Extras/VisualStudioSnippets/ue4deldyn1.snippet b/Engine/Extras/VisualStudioSnippets/ue4deldyn1.snippet
deleted file mode 100644
diff --git a/Engine/Extras/VisualStudioSnippets/ue4deldyn2.snippet b/Engine/Extras/VisualStudioSnippets/ue4deldyn2.snippet
deleted file mode 100644
diff --git a/Engine/Extras/VisualStudioSnippets/ue4deldyn3.snippet b/Engine/Extras/VisualStudioSnippets/ue4deldyn3.snippet
deleted file mode 100644
diff --git a/Engine/Extras/VisualStudioSnippets/ue4deldyn4.snippet b/Engine/Extras/VisualStudioSnippets/ue4deldyn4.snippet
deleted file mode 100644
diff --git a/Engine/Extras/VisualStudioSnippets/ue4deldyn5.snippet b/Engine/Extras/VisualStudioSnippets/ue4deldyn5.snippet
deleted file mode 100644
diff --git a/Engine/Extras/VisualStudioSnippets/ue4deldyn6.snippet b/Engine/Extras/VisualStudioSnippets/ue4deldyn6.snippet
deleted file mode 100644
diff --git a/Engine/Extras/VisualStudioSnippets/ue4deldyn7.snippet b/Engine/Extras/VisualStudioSnippets/ue4deldyn7.snippet
deleted file mode 100644
diff --git a/Engine/Extras/VisualStudioSnippets/ue4deldyn8.snippet b/Engine/Extras/VisualStudioSnippets/ue4deldyn8.snippet
deleted file mode 100644
diff --git a/Engine/Extras/VisualStudioSnippets/ue4deldynmul.snippet b/Engine/Extras/VisualStudioSnippets/ue4deldynmul.snippet
deleted file mode 100644
diff --git a/Engine/Extras/VisualStudioSnippets/ue4deldynmul1.snippet b/Engine/Extras/VisualStudioSnippets/ue4deldynmul1.snippet
deleted file mode 100644
diff --git a/Engine/Extras/VisualStudioSnippets/ue4deldynmul2.snippet b/Engine/Extras/VisualStudioSnippets/ue4deldynmul2.snippet
deleted file mode 100644
diff --git a/Engine/Extras/VisualStudioSnippets/ue4deldynmul3.snippet b/Engine/Extras/VisualStudioSnippets/ue4deldynmul3.snippet
deleted file mode 100644
diff --git a/Engine/Extras/VisualStudioSnippets/ue4deldynmul4.snippet b/Engine/Extras/VisualStudioSnippets/ue4deldynmul4.snippet
deleted file mode 100644
diff --git a/Engine/Extras/VisualStudioSnippets/ue4deldynmul5.snippet b/Engine/Extras/VisualStudioSnippets/ue4deldynmul5.snippet
deleted file mode 100644
diff --git a/Engine/Extras/VisualStudioSnippets/ue4deldynmul6.snippet b/Engine/Extras/VisualStudioSnippets/ue4deldynmul6.snippet
deleted file mode 100644
diff --git a/Engine/Extras/VisualStudioSnippets/ue4deldynmul7.snippet b/Engine/Extras/VisualStudioSnippets/ue4deldynmul7.snippet
deleted file mode 100644
diff --git a/Engine/Extras/VisualStudioSnippets/ue4deldynmul8.snippet b/Engine/Extras/VisualStudioSnippets/ue4deldynmul8.snippet
deleted file mode 100644
diff --git a/Engine/Extras/VisualStudioSnippets/ue4delmul.snippet b/Engine/Extras/VisualStudioSnippets/ue4delmul.snippet
deleted file mode 100644
diff --git a/Engine/Extras/VisualStudioSnippets/ue4delmul1.snippet b/Engine/Extras/VisualStudioSnippets/ue4delmul1.snippet
deleted file mode 100644
diff --git a/Engine/Extras/VisualStudioSnippets/ue4delmul2.snippet b/Engine/Extras/VisualStudioSnippets/ue4delmul2.snippet
deleted file mode 100644
diff --git a/Engine/Extras/VisualStudioSnippets/ue4delmul3.snippet b/Engine/Extras/VisualStudioSnippets/ue4delmul3.snippet
deleted file mode 100644
diff --git a/Engine/Extras/VisualStudioSnippets/ue4delmul4.snippet b/Engine/Extras/VisualStudioSnippets/ue4delmul4.snippet
deleted file mode 100644
diff --git a/Engine/Extras/VisualStudioSnippets/ue4delmul5.snippet b/Engine/Extras/VisualStudioSnippets/ue4delmul5.snippet
deleted file mode 100644
diff --git a/Engine/Extras/VisualStudioSnippets/ue4delmul6.snippet b/Engine/Extras/VisualStudioSnippets/ue4delmul6.snippet
deleted file mode 100644
diff --git a/Engine/Extras/VisualStudioSnippets/ue4delmul7.snippet b/Engine/Extras/VisualStudioSnippets/ue4delmul7.snippet
deleted file mode 100644
diff --git a/Engine/Extras/VisualStudioSnippets/ue4delmul8.snippet b/Engine/Extras/VisualStudioSnippets/ue4delmul8.snippet
deleted file mode 100644
diff --git a/Engine/Extras/VisualStudioSnippets/ue4enum.snippet b/Engine/Extras/VisualStudioSnippets/ue4enum.snippet
deleted file mode 100644
diff --git a/Engine/Extras/VisualStudioSnippets/ue4enumdisplay.snippet b/Engine/Extras/VisualStudioSnippets/ue4enumdisplay.snippet
deleted file mode 100644
diff --git a/Engine/Extras/VisualStudioSnippets/ue4eve.snippet b/Engine/Extras/VisualStudioSnippets/ue4eve.snippet
deleted file mode 100644
diff --git a/Engine/Extras/VisualStudioSnippets/ue4eve1.snippet b/Engine/Extras/VisualStudioSnippets/ue4eve1.snippet
deleted file mode 100644
diff --git a/Engine/Extras/VisualStudioSnippets/ue4eve2.snippet b/Engine/Extras/VisualStudioSnippets/ue4eve2.snippet
deleted file mode 100644
diff --git a/Engine/Extras/VisualStudioSnippets/ue4eve3.snippet b/Engine/Extras/VisualStudioSnippets/ue4eve3.snippet
deleted file mode 100644
diff --git a/Engine/Extras/VisualStudioSnippets/ue4eve4.snippet b/Engine/Extras/VisualStudioSnippets/ue4eve4.snippet
deleted file mode 100644
diff --git a/Engine/Extras/VisualStudioSnippets/ue4eve5.snippet b/Engine/Extras/VisualStudioSnippets/ue4eve5.snippet
deleted file mode 100644
diff --git a/Engine/Extras/VisualStudioSnippets/ue4eve6.snippet b/Engine/Extras/VisualStudioSnippets/ue4eve6.snippet
deleted file mode 100644
diff --git a/Engine/Extras/VisualStudioSnippets/ue4eve7.snippet b/Engine/Extras/VisualStudioSnippets/ue4eve7.snippet
deleted file mode 100644
diff --git a/Engine/Extras/VisualStudioSnippets/ue4eve8.snippet b/Engine/Extras/VisualStudioSnippets/ue4eve8.snippet
deleted file mode 100644
diff --git a/Engine/Extras/VisualStudioSnippets/ue4interface.snippet b/Engine/Extras/VisualStudioSnippets/ue4interface.snippet
deleted file mode 100644
diff --git a/Engine/Extras/VisualStudioSnippets/ue4log.snippet b/Engine/Extras/VisualStudioSnippets/ue4log.snippet
deleted file mode 100644
diff --git a/Engine/Extras/VisualStudioSnippets/ue4logdeclare.snippet b/Engine/Extras/VisualStudioSnippets/ue4logdeclare.snippet
deleted file mode 100644
diff --git a/Engine/Extras/VisualStudioSnippets/ue4logdefine.snippet b/Engine/Extras/VisualStudioSnippets/ue4logdefine.snippet
deleted file mode 100644
diff --git a/Engine/Extras/VisualStudioSnippets/ue4logfloat.snippet b/Engine/Extras/VisualStudioSnippets/ue4logfloat.snippet
deleted file mode 100644
diff --git a/Engine/Extras/VisualStudioSnippets/ue4logint.snippet b/Engine/Extras/VisualStudioSnippets/ue4logint.snippet
deleted file mode 100644
diff --git a/Engine/Extras/VisualStudioSnippets/ue4loguobj.snippet b/Engine/Extras/VisualStudioSnippets/ue4loguobj.snippet
deleted file mode 100644
diff --git a/Engine/Extras/VisualStudioSnippets/ue4mark.snippet b/Engine/Extras/VisualStudioSnippets/ue4mark.snippet
deleted file mode 100644
diff --git a/Engine/Extras/VisualStudioSnippets/ue4struct.snippet b/Engine/Extras/VisualStudioSnippets/ue4struct.snippet
deleted file mode 100644
