Add code snippet to disable optimizations around a block of code

#rnx

diff --git a/Engine/Extras/VisualStudioSnippets/Unreal.deoptimize.snippet b/Engine/Extras/VisualStudioSnippets/Unreal.deoptimize.snippet
new file mode 100644
--- /dev/null
+++ b/Engine/Extras/VisualStudioSnippets/Unreal.deoptimize.snippet
@@ -0,0 +1,21 @@
+﻿<?xml version="1.0" encoding="utf-8" ?>
+<CodeSnippets  xmlns="http://schemas.microsoft.com/VisualStudio/2005/CodeSnippet">
+	<CodeSnippet Format="1.0.0">
+		<Header>
+			<Title>unreal-deoptimize</Title>
+			<Shortcut>unreal-deoptimize</Shortcut>
+			<Description>Code snippet for disabling optimizations</Description>
+			<Author>Epic Games</Author>
+			<SnippetTypes>
+				<SnippetType>Expansion</SnippetType>
+				<SnippetType>SurroundsWith</SnippetType>
+			</SnippetTypes>
+		</Header>
+		<Snippet>
+			<Code Language="cpp"><![CDATA[UE_DISABLE_OPTIMIZATION
+			$selected$
+			UE_ENABLE_OPTIMIZATION]]>
+			</Code>
+		</Snippet>
+	</CodeSnippet>
+</CodeSnippets>
\ No newline at end of file
